Experimental MultiServerReplication plugin & some simple test cases.
Exposes a peer-to-peer like interface for communicating between multiple UE instances via online beacons.

#virtualized

diff --git a/Engine/Plugins/Runtime/MultiServerReplication/MultiServerReplication.uplugin b/Engine/Plugins/Runtime/MultiServerReplication/MultiServerReplication.uplugin
new file mode 100644
--- /dev/null
+++ b/Engine/Plugins/Runtime/MultiServerReplication/MultiServerReplication.uplugin
@@ -0,0 +1,27 @@
+{
+	"FileVersion": 3,
+	"FriendlyName": "MultiServerReplication Plugin",
+	"Version": 1,
+	"VersionName": "1.0",
+	"Description": "",
+	"Category": "Epic Internal",
+	"CreatedBy": "Epic Games, Inc.",
+	"CreatedByURL": "http://epicgames.com",
+	"EnabledByDefault": false,
+	"CanContainContent": false,
+	"IsExperimentalVersion":  true,
+	"Installed": false,
+	"Modules": [
+		{
+			"Name": "MultiServerReplication",
+			"Type": "Runtime",
+			"LoadingPhase": "Default"
+		}
+	],
+	"Plugins": [
+		{
+			"Name": "OnlineSubsystemUtils",
+			"Enabled": true
+		}
+	]
+}
\ No newline at end of file
diff --git a/Engine/Plugins/Runtime/MultiServerReplication/Source/MultiServerReplication/MultiServerReplication.Build.cs b/Engine/Plugins/Runtime/MultiServerReplication/Source/MultiServerReplication/MultiServerReplication.Build.cs
new file mode 100644
--- /dev/null
+++ b/Engine/Plugins/Runtime/MultiServerReplication/Source/MultiServerReplication/MultiServerReplication.Build.cs
@@ -0,0 +1,24 @@
+// Copyright Epic Games, Inc. All Rights Reserved.
+
+using UnrealBuildTool;
+
+public class MultiServerReplication : ModuleRules
+{
+	public MultiServerReplication(ReadOnlyTargetRules Target) : base(Target)
+    {
+        PrivateDependencyModuleNames.AddRange(
+			new string[] {
+                "Core",
+                "CoreUObject",
+                "Engine",
+				"NetCore"
+            }
+        );
+
+        PublicDependencyModuleNames.AddRange(
+            new string[] {
+				"OnlineSubsystemUtils"
+			}
+		);
+    }
+}
diff --git a/Engine/Plugins/Runtime/MultiServerReplication/Source/MultiServerReplication/Private/MultiServerBeaconClient.cpp b/Engine/Plugins/Runtime/MultiServerReplication/Source/MultiServerReplication/Private/MultiServerBeaconClient.cpp
new file mode 100644
--- /dev/null
+++ b/Engine/Plugins/Runtime/MultiServerReplication/Source/MultiServerReplication/Private/MultiServerBeaconClient.cpp
@@ -0,0 +1,220 @@
+// Copyright Epic Games, Inc. All Rights Reserved.
+
+#include "MultiServerBeaconClient.h"
+#include "MultiServerBeaconHostObject.h"
+#include "MultiServerNode.h"
+#include "Engine/Engine.h"
+#include "Engine/Level.h"
+#include "Engine/LevelStreaming.h"
+#include "Engine/NetConnection.h"
+#include "Engine/World.h"
+#include "Misc/PackageName.h"
+#include "Net/UnrealNetwork.h"
+
+#include UE_INLINE_GENERATED_CPP_BY_NAME(MultiServerBeaconClient)
+
+AMultiServerBeaconClient::AMultiServerBeaconClient() :
+	Super()
+{
+	bOnlyRelevantToOwner = true;
+
+	NetDriverName = FName(TEXT("MultiServerNetDriverClient"));
+	NetDriverDefinitionName = FName(TEXT("MultiServerNetDriver"));
+
+	// Allow this to tick on MultiServer nodes
+	PrimaryActorTick.bAllowTickOnDedicatedServer = true;
+}
+
+FName AMultiServerBeaconClient::NetworkRemapPath(FName InPackageName, bool bReading)
+{
+	// For PIE Networking: remap the packagename to our local PIE packagename
+	FString PackageNameStr = InPackageName.ToString();
+	if (UNetConnection* Connection = GetNetConnection())
+	{
+		GEngine->NetworkRemapPath(Connection, PackageNameStr, bReading);
+	}
+	return FName(*PackageNameStr);
+}
+
+void AMultiServerBeaconClient::OnConnected()
+{
+	Super::OnConnected();
+
+	UE_LOG(LogMultiServerBeacon, Log, TEXT("MultiServer beacon connection established."));
+
+	if (ensureMsgf(OwningNode, TEXT("No owning node")))
+	{
+		ServerSetRemotePeerId(OwningNode->GetLocalPeerId());
+	}
+
+	UWorld* World = GetWorld();
+	if (World)
+	{
+		TArray<FUpdateLevelVisibilityLevelInfo> LevelVisibilities;
+		for (ULevelStreaming* LevelStreaming : World->GetStreamingLevels())
+		{
+			if (LevelStreaming)
+			{
+				const ULevel* Level = LevelStreaming->GetLoadedLevel();
+				if (Level && Level->bIsVisible && !Level->bClientOnlyVisible)
+				{
+					FUpdateLevelVisibilityLevelInfo& LevelVisibility = *new(LevelVisibilities) FUpdateLevelVisibilityLevelInfo(Level, true);
+					LevelVisibility.PackageName = NetworkRemapPath(LevelVisibility.PackageName, false);
+				}
+			}
+		}
+		if (LevelVisibilities.Num() > 0)
+		{
+			ServerUpdateMultipleLevelsVisibility(LevelVisibilities);
+		}
+	}
+
+	OnLevelRemovedFromWorldHandle = FWorldDelegates::LevelRemovedFromWorld.AddUObject(this, &ThisClass::OnLevelRemovedFromWorld);
+	OnLevelAddedToWorldHandle = FWorldDelegates::LevelAddedToWorld.AddUObject(this, &ThisClass::OnLevelAddedToWorld);
+}
+
+void AMultiServerBeaconClient::DestroyBeacon()
+{
+	FWorldDelegates::LevelRemovedFromWorld.Remove(OnLevelRemovedFromWorldHandle);
+	FWorldDelegates::LevelAddedToWorld.Remove(OnLevelAddedToWorldHandle);
+
+	Super::DestroyBeacon();
+}
+
+void AMultiServerBeaconClient::ConnectToServer(const FString& ConnectInfo)
+{
+	bool bSuccess = false;
+
+	FURL ConnectURL(NULL, *ConnectInfo, TRAVEL_Absolute);
+	if (InitClient(ConnectURL))
+	{
+		bSuccess = true;
+	}
+	else
+	{
+		UE_LOG(LogMultiServerBeacon, Warning, TEXT("ConnectToRelay: Failure to init client beacon with %s."), *ConnectURL.ToString());
+	}
+
+	if (!bSuccess)
+	{
+		OnFailure();
+	}
+}
+
+void AMultiServerBeaconClient::ClientPeerConnected_Implementation(const FString& NewRemotePeerId, AMultiServerBeaconClient* Beacon)
+{
+	RemotePeerId = NewRemotePeerId;
+
+	if (OwningNode)
+	{
+		OwningNode->OnMultiServerConnected.ExecuteIfBound(OwningNode->GetLocalPeerId(), NewRemotePeerId, this);
+	}
+}
+
+void AMultiServerBeaconClient::OnFailure()
+{
+	Super::OnFailure();
+}
+
+void AMultiServerBeaconClient::ServerUpdateLevelVisibility_Implementation(const FUpdateLevelVisibilityLevelInfo& LevelVisibility)
+{
+	if (GetConnectionState() == EBeaconConnectionState::Open)
+	{
+		if (UNetConnection* Connection = GetNetConnection())
+		{
+			if (Connection->Driver && Connection->Driver->IsServer())
+			{
+				FUpdateLevelVisibilityLevelInfo LevelVisibilityCopy = LevelVisibility;
+				LevelVisibilityCopy.PackageName = NetworkRemapPath(LevelVisibility.PackageName, true);
+				LevelVisibilityCopy.bSkipCloseOnError = true;
+				
+				Connection->UpdateLevelVisibility(LevelVisibilityCopy);
+			}
+		}
+	}
+}
+
+bool AMultiServerBeaconClient::ServerUpdateLevelVisibility_Validate(const FUpdateLevelVisibilityLevelInfo& LevelVisibility)
+{
+	RPC_VALIDATE(LevelVisibility.PackageName.IsValid());
+
+	FText Reason;
+
+	if (!FPackageName::IsValidLongPackageName(LevelVisibility.PackageName.ToString(), true, &Reason))
+	{
+		UE_LOG(LogMultiServerBeacon, Warning, TEXT("ServerUpdateLevelVisibility() Invalid package name: %s (%s)"), *LevelVisibility.PackageName.ToString(), *Reason.ToString());
+		return false;
+	}
+
+	return true;
+}
+
+void AMultiServerBeaconClient::ServerUpdateMultipleLevelsVisibility_Implementation(const TArray<FUpdateLevelVisibilityLevelInfo>& LevelVisibilities)
+{
+	for (const FUpdateLevelVisibilityLevelInfo& LevelVisibility : LevelVisibilities)
+	{
+		ServerUpdateLevelVisibility_Implementation(LevelVisibility);
+	}
+}
+
+bool AMultiServerBeaconClient::ServerUpdateMultipleLevelsVisibility_Validate(const TArray<FUpdateLevelVisibilityLevelInfo>& LevelVisibilities)
+{
+	for (const FUpdateLevelVisibilityLevelInfo& LevelVisibility : LevelVisibilities)
+	{
+		if (!ServerUpdateLevelVisibility_Validate(LevelVisibility))
+		{
+			return false;
+		}
+	}
+
+	return true;
+}
+
+void AMultiServerBeaconClient::ServerSetRemotePeerId_Implementation(const FString& NewRemoteId)
+{
+	UE_LOG(LogMultiServerBeacon, Log, TEXT("ServerSetRemotePeerId: %s %s"), *GetNameSafe(this), *NewRemoteId);
+	RemotePeerId = NewRemoteId;
+
+	if (OwningNode)
+	{
+		OwningNode->OnMultiServerConnected.ExecuteIfBound(OwningNode->GetLocalPeerId(), NewRemoteId, this);
+	}
+}
+
+FString AMultiServerBeaconClient::GetRemotePeerId() const
+{
+	return RemotePeerId;
+}
+
+bool AMultiServerBeaconClient::IsAuthorityBeacon() const
+{
+	return !GetNetDriver() || GetNetDriver()->ServerConnection == nullptr;
+}
+
+void AMultiServerBeaconClient::OnLevelRemovedFromWorld(ULevel* Level, UWorld* World)
+{
+	if (GetWorld() == World)
+	{
+		if (Level && !Level->bClientOnlyVisible)
+		{
+			FUpdateLevelVisibilityLevelInfo LevelVisibility(Level, false);
+			LevelVisibility.PackageName = NetworkRemapPath(LevelVisibility.PackageName, false);
+
+			ServerUpdateLevelVisibility(LevelVisibility);
+		}
+	}
+}
+
+void AMultiServerBeaconClient::OnLevelAddedToWorld(ULevel* Level, UWorld* World)
+{
+	if (GetWorld() == World)
+	{
+		if (Level && !Level->bClientOnlyVisible)
+		{
+			FUpdateLevelVisibilityLevelInfo LevelVisibility(Level, true);
+			LevelVisibility.PackageName = NetworkRemapPath(LevelVisibility.PackageName, false);
+
+			ServerUpdateLevelVisibility(LevelVisibility);
+		}
+	}
+}
\ No newline at end of file
diff --git a/Engine/Plugins/Runtime/MultiServerReplication/Source/MultiServerReplication/Private/MultiServerBeaconHost.cpp b/Engine/Plugins/Runtime/MultiServerReplication/Source/MultiServerReplication/Private/MultiServerBeaconHost.cpp
new file mode 100644
--- /dev/null
+++ b/Engine/Plugins/Runtime/MultiServerReplication/Source/MultiServerReplication/Private/MultiServerBeaconHost.cpp
@@ -0,0 +1,93 @@
+// Copyright Epic Games, Inc. All Rights Reserved.
+
+#include "MultiServerBeaconHost.h"
+#include "Misc/CommandLine.h"
+#include "Net/DataChannel.h"
+#include "Engine/NetConnection.h"
+
+#include UE_INLINE_GENERATED_CPP_BY_NAME(MultiServerBeaconHost)
+
+AMultiServerBeaconHost::AMultiServerBeaconHost(const FObjectInitializer& ObjectInitializer) :
+	Super(ObjectInitializer)
+{
+	NetDriverName = FName(TEXT("MultiServerNetDriverHost"));
+	NetDriverDefinitionName = FName(TEXT("MultiServerNetDriver"));
+	MaxConnections = 100;
+}
+
+bool AMultiServerBeaconHost::InitHost()
+{
+	FURL URL(nullptr, TEXT(""), TRAVEL_Absolute);
+
+	FString MultiServerBeaconHostAddr;
+	if (FParse::Value(FCommandLine::Get(), TEXT("MultiServerHostAddr="), MultiServerBeaconHostAddr) && !MultiServerBeaconHostAddr.IsEmpty())
+	{
+		URL.AddOption(*FString::Printf(TEXT("multihome=%s"), *MultiServerBeaconHostAddr));
+	}
+
+	URL.Port = ListenPort;
+	if(URL.Valid)
+	{
+		if (InitBase() && NetDriver)
+		{
+			FString Error;
+			if (NetDriver->InitListen(this, URL, bReuseAddressAndPort, Error))
+			{
+				ListenPort = URL.Port;
+				NetDriver->SetWorld(GetWorld());
+				NetDriver->Notify = this;
+				NetDriver->InitialConnectTimeout = BeaconConnectionInitialTimeout;
+				NetDriver->ConnectionTimeout = BeaconConnectionTimeout;
+				return true;
+			}
+			else
+			{
+				// error initializing the network stack...
+				UE_LOG(LogBeacon, Log, TEXT("AMultiServerBeaconHost::InitHost failed"));
+				OnFailure();
+			}
+		}
+	}
+
+	return false;
+}
+
+bool AMultiServerBeaconHost::AtCapacity() const
+{
+	int32 Count = 0;
+
+	for (UNetConnection* Connection : NetDriver->ClientConnections)
+	{
+		if (Connection && (Connection->GetConnectionState() != USOCK_Closed) && (Connection->OwningActor != nullptr))
+		{
+			++Count;
+		}
+	}
+
+	return (Count >= MaxConnections);
+}
+
+void AMultiServerBeaconHost::NotifyControlMessage(UNetConnection* Connection, uint8 MessageType, FInBunch& Bunch)
+{
+	if (NetDriver->ServerConnection == nullptr)
+	{
+		switch (MessageType)
+		{
+		case NMT_BeaconJoin:
+			// if we are at capacity, intervene here
+			if (AtCapacity())
+			{
+				FString ErrorMsg = TEXT("MultiServer beacon at capacity.");
+				UE_LOG(LogBeacon, Log, TEXT("%s: %s"), *Connection->GetName(), *ErrorMsg);
+				FNetControlMessage<NMT_Failure>::Send(Connection, ErrorMsg);
+				Connection->FlushNet(true);
+				Connection->Close();
+				return;
+			}
+			break;
+		}
+	}
+
+	Super::NotifyControlMessage(Connection, MessageType, Bunch);
+}
+
diff --git a/Engine/Plugins/Runtime/MultiServerReplication/Source/MultiServerReplication/Private/MultiServerBeaconHost.h b/Engine/Plugins/Runtime/MultiServerReplication/Source/MultiServerReplication/Private/MultiServerBeaconHost.h
new file mode 100644
--- /dev/null
+++ b/Engine/Plugins/Runtime/MultiServerReplication/Source/MultiServerReplication/Private/MultiServerBeaconHost.h
@@ -0,0 +1,29 @@
+// Copyright Epic Games, Inc. All Rights Reserved.
+
+#pragma once
+
+#include "OnlineBeaconHost.h"
+#include "MultiServerBeaconHost.generated.h"
+
+/**
+ * An online beacon that helps manage connecting to MultiServer networks, and replicating
+ * metadata about the MultiServer network.
+ */
+UCLASS(transient, config=Engine)
+class AMultiServerBeaconHost : public AOnlineBeaconHost
+{
+	GENERATED_BODY()
+
+public:
+	AMultiServerBeaconHost(const FObjectInitializer& ObjectInitializer);
+
+	virtual bool InitHost() override;
+	virtual void NotifyControlMessage(UNetConnection* Connection, uint8 MessageType, FInBunch& Bunch) override;
+
+	/** @return Whether or not this Node already has the maximum number of allowable connections. */
+	virtual bool AtCapacity() const;
+
+protected:
+	UPROPERTY(config)
+	int32 MaxConnections;
+};
diff --git a/Engine/Plugins/Runtime/MultiServerReplication/Source/MultiServerReplication/Private/MultiServerBeaconHostObject.cpp b/Engine/Plugins/Runtime/MultiServerReplication/Source/MultiServerReplication/Private/MultiServerBeaconHostObject.cpp
new file mode 100644
--- /dev/null
+++ b/Engine/Plugins/Runtime/MultiServerReplication/Source/MultiServerReplication/Private/MultiServerBeaconHostObject.cpp
@@ -0,0 +1,42 @@
+// Copyright Epic Games, Inc. All Rights Reserved.
+
+#include "MultiServerBeaconHostObject.h"
+#include "MultiServerBeaconClient.h"
+#include "MultiServerNode.h"
+
+#include UE_INLINE_GENERATED_CPP_BY_NAME(MultiServerBeaconHostObject)
+
+DEFINE_LOG_CATEGORY(LogMultiServerBeacon);
+
+AMultiServerBeaconHostObject::AMultiServerBeaconHostObject(const FObjectInitializer& ObjectInitializer) :
+	Super(ObjectInitializer)
+{
+	SetClientBeaconActorClass(AMultiServerBeaconClient::StaticClass());
+}
+
+void AMultiServerBeaconHostObject::SetClientBeaconActorClass(TSubclassOf<AOnlineBeaconClient> InClientBeaconActorClass)
+{
+	ClientBeaconActorClass = InClientBeaconActorClass;
+	BeaconTypeName = ClientBeaconActorClass->GetName();
+}
+
+void AMultiServerBeaconHostObject::OnClientConnected(AOnlineBeaconClient* NewClientActor, UNetConnection* ClientConnection)
+{
+	Super::OnClientConnected(NewClientActor, ClientConnection);
+
+	UE_LOG(LogMultiServerBeacon, Verbose, TEXT("AMultiServerBeaconHostObject::OnClientConnected: num client actors connected is %d"), ClientActors.Num());
+	if (AMultiServerBeaconClient* MultiServerClient = Cast<AMultiServerBeaconClient>(NewClientActor))
+	{
+		MultiServerClient->SetOwningNode(OwningNode);
+
+		MultiServerClient->ClientPeerConnected(OwningNode->GetLocalPeerId(), MultiServerClient);
+	}
+}
+
+void AMultiServerBeaconHostObject::NotifyClientDisconnected(AOnlineBeaconClient* LeavingClientActor)
+{
+	Super::NotifyClientDisconnected(LeavingClientActor);
+
+	UE_LOG(LogMultiServerBeacon, Verbose, TEXT("AMultiServerBeaconHostObject::NotifyClientDisconnected: num client actors connected is %d"), ClientActors.Num());
+}
+
diff --git a/Engine/Plugins/Runtime/MultiServerReplication/Source/MultiServerReplication/Private/MultiServerBeaconHostObject.h b/Engine/Plugins/Runtime/MultiServerReplication/Source/MultiServerReplication/Private/MultiServerBeaconHostObject.h
new file mode 100644
--- /dev/null
+++ b/Engine/Plugins/Runtime/MultiServerReplication/Source/MultiServerReplication/Private/MultiServerBeaconHostObject.h
@@ -0,0 +1,32 @@
+// Copyright Epic Games, Inc. All Rights Reserved.
+
+#pragma once
+
+#include "OnlineBeaconHostObject.h"
+#include "MultiServerBeaconHostObject.generated.h"
+
+DECLARE_LOG_CATEGORY_EXTERN(LogMultiServerBeacon, Log, All);
+
+class AMultiServerBeaconClient;
+class UMultiServerNode;
+
+UCLASS(transient, config=Engine)
+class AMultiServerBeaconHostObject : public AOnlineBeaconHostObject
+{
+	GENERATED_BODY()
+
+public:
+	AMultiServerBeaconHostObject(const FObjectInitializer& ObjectInitializer);
+
+	void SetClientBeaconActorClass(TSubclassOf<AOnlineBeaconClient> InClientBeaconActorClass);
+	void SetOwningNode(UMultiServerNode* InOwningNode) { OwningNode = InOwningNode; }
+
+	virtual void OnClientConnected(AOnlineBeaconClient* NewClientActor, UNetConnection* ClientConnection) override;
+
+	virtual void NotifyClientDisconnected(AOnlineBeaconClient* LeavingClientActor) override;
+
+private:
+
+	UPROPERTY()
+	TObjectPtr<UMultiServerNode> OwningNode;
+};
diff --git a/Engine/Plugins/Runtime/MultiServerReplication/Source/MultiServerReplication/Private/MultiServerNode.cpp b/Engine/Plugins/Runtime/MultiServerReplication/Source/MultiServerReplication/Private/MultiServerNode.cpp
new file mode 100644
--- /dev/null
+++ b/Engine/Plugins/Runtime/MultiServerReplication/Source/MultiServerReplication/Private/MultiServerNode.cpp
@@ -0,0 +1,224 @@
+// Copyright Epic Games, Inc. All Rights Reserved.
+
+#include "MultiServerNode.h"
+#include "MultiServerBeaconHost.h"
+#include "MultiServerBeaconHostObject.h"
+#include "MultiServerBeaconClient.h"
+#include "MultiServerPeerConnection.h"
+#include "MultiServerReplicationTypes.h"
+#include "Engine/NetConnection.h"
+#include "Engine/World.h"
+#include "Misc/Parse.h"
+#include "Misc/CommandLine.h"
+
+#include UE_INLINE_GENERATED_CPP_BY_NAME(MultiServerNode)
+
+void UMultiServerNode::ParseCommandLineIntoCreateParams(FMultiServerNodeCreateParams& InOutParams)
+{
+	FParse::Value(FCommandLine::Get(), TEXT("MultiServerLocalId="), InOutParams.LocalPeerId);
+	FParse::Value(FCommandLine::Get(), TEXT("MultiServerListenPort="), InOutParams.ListenPort);
+
+	FString PeerAddressesString;
+	FParse::Value(FCommandLine::Get(), TEXT("MultiServerPeers="), PeerAddressesString, false);
+
+	PeerAddressesString.ParseIntoArray(InOutParams.PeerAddresses, TEXT(","), true);
+}
+
+UMultiServerNode::UMultiServerNode()
+	: UObject()
+{
+	RetryConnectDelay = 0.5f;
+	RetryConnectMaxDelay = 30.0f;
+}
+
+UMultiServerNode* UMultiServerNode::Create(const FMultiServerNodeCreateParams& Params)
+{
+	UMultiServerNode* NewNode = NewObject<UMultiServerNode>(Params.World);
+	const bool bResult = NewNode->RegisterServer(Params);
+
+	if (bResult)
+	{
+		return NewNode;
+	}
+	else
+	{
+		// Is this the best way to clean up on failed registration?
+		NewNode->MarkAsGarbage();
+	}
+	
+	return nullptr;
+}
+
+bool UMultiServerNode::RegisterServer(const FMultiServerNodeCreateParams& Params)
+{
+	if (Params.World == nullptr)
+	{
+		UE_LOG(LogMultiServerReplication, Warning, TEXT("UMultiServerNode::RegisterServer: null world - failed to register."));
+		return false;
+	}
+
+	if (Params.LocalPeerId.IsEmpty())
+	{
+		UE_LOG(LogMultiServerReplication, Warning, TEXT("UMultiServerNode::RegisterServer: no MultiServerLocalId specified - required for multiserver to work properly."));
+		return false;
+	}
+
+	LocalPeerId = Params.LocalPeerId;
+	UserBeaconClass = Params.UserBeaconClass;
+
+	OnMultiServerConnected = Params.OnMultiServerConnected;
+
+	if (Params.ListenPort == 0)
+	{
+		UE_LOG(LogMultiServerReplication, Log, TEXT("UMultiServerNode::RegisterServer: no listen port specified, not listening."));
+	}
+	else
+	{
+		UE_LOG(LogMultiServerReplication, Log, TEXT("UMultiServerNode::RegisterServer: setting up host beacon for %s."), ToCStr(LocalPeerId));
+
+		// Set up host beacon
+		if (ensureMsgf(BeaconHost == nullptr, TEXT("UMultiServerNode::RegisterServer: BeaconHost already created.")))
+		{
+			// Always create a new beacon host, state will be determined in a moment
+			BeaconHost = Params.World->SpawnActor<AMultiServerBeaconHost>(AMultiServerBeaconHost::StaticClass());
+			check(BeaconHost);
+
+			BeaconHost->ListenPort = Params.ListenPort;
+
+			bool bBeaconInit = false;
+			if (BeaconHost->InitHost())
+			{ 
+				BeaconHostObject = Params.World->SpawnActor<AMultiServerBeaconHostObject>(AMultiServerBeaconHostObject::StaticClass());
+				check(BeaconHostObject);
+
+				BeaconHostObject->SetClientBeaconActorClass(Params.UserBeaconClass);
+				BeaconHostObject->SetOwningNode(this);
+
+				BeaconHost->RegisterHost(BeaconHostObject);
+				BeaconHost->PauseBeaconRequests(false);
+			}
+			else
+			{
+				UE_LOG(LogMultiServerReplication, Warning, TEXT("Failed to init multiserver host beacon %s"), *BeaconHost->GetName());
+				return false;
+			}
+		}
+	}
+
+	// If the -MultiServerPeers command-line option is set, start a client beacon for each one and connect to them.
+	// (they are expected to be listening already)
+	if (Params.PeerAddresses.Num() == 0)
+	{
+		UE_LOG(LogMultiServerReplication, Log, TEXT("UMultiServerNode::RegisterServer: no peers specified, not connecting to any. LocalPeerId %s"), ToCStr(LocalPeerId));
+	}
+	else
+	{
+		for (const FString& PeerAddress : Params.PeerAddresses)
+		{
+			FURL PeerURL(nullptr, ToCStr(PeerAddress), ETravelType::TRAVEL_Absolute);
+			
+			if (!PeerURL.Valid)
+			{
+				UE_LOG(LogMultiServerReplication, Verbose, TEXT("Failed to parse peer address %s, not connecting."), ToCStr(PeerAddress));
+				continue;
+			}
+
+			// Only connect if the peer port is "lower" than the listening port, to prevent redundant connections.
+			// This only works when connecting to instances on the same machine, so limit it to loopback addresses.
+			// TODO: Handle this for remote addresses, or require the filtering to be done at a higher level (command line, etc).
+			if (PeerURL.Host.Equals(TEXT("localhost"), ESearchCase::IgnoreCase) ||
+				PeerURL.Host.StartsWith(TEXT("127."), ESearchCase::IgnoreCase))
+			{
+				if (Params.ListenPort != 0 && PeerURL.Port >= Params.ListenPort)
+				{
+					continue;
+				}
+			}
+
+			if (!PeerAddress.IsEmpty())
+			{
+				UMultiServerPeerConnection* Peer = NewObject<UMultiServerPeerConnection>(this);
+				Peer->SetOwningNode(this);
+				Peer->SetRemoteAddress(PeerAddress);
+				Peer->SetLocalPeerId(LocalPeerId);
+				Peer->InitClientBeacon();
+				PeerConnections.Add(Peer);
+			}
+		}
+	}
+
+	return true;
+}
+
+AMultiServerBeaconClient* UMultiServerNode::GetBeaconClientForRemotePeer(FStringView RemotePeerId)
+{
+	// See if we are the host of the target server
+	UNetDriver* HostNetDriver = BeaconHost ? BeaconHost->GetNetDriver() : nullptr;
+	if (HostNetDriver)
+	{
+		for (UNetConnection* ClientConnection : HostNetDriver->ClientConnections)
+		{
+			AMultiServerBeaconClient* BeaconClient = Cast<AMultiServerBeaconClient>(BeaconHost->GetClientActor(ClientConnection));
+			if (BeaconClient)
+			{
+				if (RemotePeerId.Equals(BeaconClient->GetRemotePeerId(), ESearchCase::IgnoreCase))
+				{
+					return BeaconClient;
+				}
+			}
+		}
+	}
+
+	// See if we are a client of the target server
+	for (UMultiServerPeerConnection* Peer : PeerConnections)
+	{
+		AMultiServerBeaconClient* Beacon = Peer->BeaconClient;
+		if (Beacon)
+		{
+			if (RemotePeerId.Equals(Beacon->GetRemotePeerId(), ESearchCase::IgnoreCase))
+			{
+				return Beacon;
+			}
+		}
+	}
+
+	return nullptr;
+}
+
+AMultiServerBeaconClient* UMultiServerNode::GetBeaconClientForURL(const FString& InURL)
+{
+	FURL URL(nullptr, ToCStr(InURL), TRAVEL_Absolute);
+
+	// See if we are the host of the target server
+	UNetDriver* HostNetDriver = BeaconHost ? BeaconHost->GetNetDriver() : nullptr;
+	if (HostNetDriver)
+	{
+		for (UNetConnection* ClientConnection : HostNetDriver->ClientConnections)
+		{
+			if (ClientConnection &&
+				ClientConnection->URL.Host.Equals(URL.Host, ESearchCase::IgnoreCase) &&
+				ClientConnection->URL.Port == URL.Port)
+			{
+				return Cast<AMultiServerBeaconClient>(BeaconHost->GetClientActor(ClientConnection));
+			}
+		}
+	}
+
+	// See if we are a client of the target server
+	for (UMultiServerPeerConnection* Peer : PeerConnections)
+	{
+		AMultiServerBeaconClient* Beacon = Peer->BeaconClient;
+		if (Beacon)
+		{
+			UNetConnection* ClientConnection = Beacon->GetNetConnection();
+			if (ClientConnection &&
+				ClientConnection->URL.Host.Equals(URL.Host, ESearchCase::IgnoreCase) &&
+				ClientConnection->URL.Port == URL.Port)
+			{
+				return Beacon;
+			}
+		}
+	}
+
+	return nullptr;
+}
diff --git a/Engine/Plugins/Runtime/MultiServerReplication/Source/MultiServerReplication/Private/MultiServerPeerConnection.cpp b/Engine/Plugins/Runtime/MultiServerReplication/Source/MultiServerReplication/Private/MultiServerPeerConnection.cpp
new file mode 100644
--- /dev/null
+++ b/Engine/Plugins/Runtime/MultiServerReplication/Source/MultiServerReplication/Private/MultiServerPeerConnection.cpp
@@ -0,0 +1,106 @@
+// Copyright Epic Games, Inc. All Rights Reserved.
+
+#include "MultiServerPeerConnection.h"
+#include "MultiServerReplicationTypes.h"
+#include "MultiServerNode.h"
+#include "MultiServerBeaconClient.h"
+#include "Engine/World.h"
+#include "TimerManager.h"
+
+#include UE_INLINE_GENERATED_CPP_BY_NAME(MultiServerPeerConnection)
+
+UMultiServerPeerConnection::UMultiServerPeerConnection()
+	: Super()
+{
+	
+}
+
+void UMultiServerPeerConnection::InitClientBeacon()
+{
+	DestroyClientBeacon();
+
+	BeaconClient = GetWorld()->SpawnActor<AMultiServerBeaconClient>(OwningNode->GetUserBeaconClass());
+	if (BeaconClient)
+	{
+		UE_LOG(LogMultiServerReplication, Verbose, TEXT("Created multiserver client beacon %s."), *BeaconClient->GetName());
+
+		BeaconClient->OnHostConnectionFailure().BindUObject(this, &ThisClass::OnBeaconConnectionFailure);
+		BeaconClient->SetOwningNode(OwningNode);
+
+		if (!RemoteAddress.IsEmpty())
+		{
+			BeaconClient->ConnectToServer(RemoteAddress);
+		}
+		else
+		{
+			UE_LOG(LogMultiServerReplication, Verbose, TEXT("Failed to get connection info for client beacon %s"), *BeaconClient->GetName());
+			OnBeaconConnectionFailure();
+		}
+	}
+	else
+	{
+		UE_LOG(LogMultiServerReplication, Warning, TEXT("Failed to init MultiServer client beacon for %s"), ToCStr(RemoteAddress));
+		OnBeaconConnectionFailure();
+	}
+}
+
+void UMultiServerPeerConnection::DestroyClientBeacon()
+{
+	ClearConnectRetryTimer();
+
+	if (BeaconClient)
+	{
+		UE_LOG(LogMultiServerReplication, Verbose, TEXT("Destroying MultiServer beacon client."));
+
+		BeaconClient->OnHostConnectionFailure().Unbind();
+
+		BeaconClient->DestroyBeacon();
+		BeaconClient = nullptr;
+	}
+}
+
+void UMultiServerPeerConnection::ClearConnectRetryTimer()
+{
+	if (UWorld* World = GetWorld())
+	{
+		World->GetTimerManager().ClearTimer(ConnectRetryTimerHandle);
+	}
+}
+
+float UMultiServerPeerConnection::GetRetryDelay()
+{
+	// Retry delay that increases exponentially
+	++ConnectAttemptNum;
+
+	// Randomize the initial delay
+	float RetryDelay = FMath::RandRange(0.1f, OwningNode->GetRetryConnectDelay());
+
+	if (ConnectAttemptNum > 1)
+	{
+		RetryDelay += OwningNode->GetRetryConnectDelay() * FMath::Pow(2.f, ConnectAttemptNum - 1);
+	}
+
+	RetryDelay = FMath::Min(RetryDelay, OwningNode->GetRetryConnectMaxDelay());
+
+	return RetryDelay;
+}
+
+void UMultiServerPeerConnection::OnBeaconConnectionFailure()
+{
+	UE_LOG(LogMultiServerReplication, Log, TEXT("MultiServer beacon connection failed."));
+
+	DestroyClientBeacon();
+
+	float Delay = GetRetryDelay();
+
+	UE_LOG(LogMultiServerReplication, Log, TEXT("MultiServer peer connect retry in %.2f seconds, attempt #%d"), Delay, ConnectAttemptNum);
+
+	UWorld* World = GetWorld();
+	if (World)
+	{
+		World->GetTimerManager().SetTimer(ConnectRetryTimerHandle, FTimerDelegate::CreateWeakLambda(this, [this]()
+			{
+				InitClientBeacon();
+			}), Delay, false);
+	}
+}
\ No newline at end of file
diff --git a/Engine/Plugins/Runtime/MultiServerReplication/Source/MultiServerReplication/Private/MultiServerPeerConnection.h b/Engine/Plugins/Runtime/MultiServerReplication/Source/MultiServerReplication/Private/MultiServerPeerConnection.h
new file mode 100644
--- /dev/null
+++ b/Engine/Plugins/Runtime/MultiServerReplication/Source/MultiServerReplication/Private/MultiServerPeerConnection.h
@@ -0,0 +1,48 @@
+// Copyright Epic Games, Inc. All Rights Reserved.
+
+#pragma once
+
+#include "UObject/Object.h"
+#include "UObject/ObjectPtr.h"
+#include "Engine/TimerHandle.h"
+#include "Templates/UnrealTemplate.h"
+#include "MultiServerPeerConnection.generated.h"
+
+class AMultiServerBeaconClient;
+class UMultiServerNode;
+
+/** Net connection class specific to MultiServer Networking. */
+UCLASS(Transient)
+class UMultiServerPeerConnection : public UObject
+{
+	GENERATED_BODY()
+
+public:
+	UMultiServerPeerConnection();
+
+	void InitClientBeacon();
+	void DestroyClientBeacon();
+	void ClearConnectRetryTimer();
+
+	float GetRetryDelay();
+
+	void OnBeaconConnectionFailure();
+
+	UPROPERTY()
+	TObjectPtr<AMultiServerBeaconClient> BeaconClient;
+
+	void SetOwningNode(UMultiServerNode* InOwningNode) { OwningNode = InOwningNode; }
+	void SetRemoteAddress(FString InRemoteAddress) { RemoteAddress = MoveTemp(InRemoteAddress); }
+	void SetLocalPeerId(FString InLocalPeerId) { LocalPeerId = InLocalPeerId; }
+
+private:
+	int ConnectAttemptNum = 0;
+
+	FTimerHandle ConnectRetryTimerHandle;
+
+	UPROPERTY()
+	TObjectPtr<UMultiServerNode> OwningNode;
+
+	FString RemoteAddress;
+	FString LocalPeerId;
+};
diff --git a/Engine/Plugins/Runtime/MultiServerReplication/Source/MultiServerReplication/Private/MultiServerReplicationModule.cpp b/Engine/Plugins/Runtime/MultiServerReplication/Source/MultiServerReplication/Private/MultiServerReplicationModule.cpp
new file mode 100644
--- /dev/null
+++ b/Engine/Plugins/Runtime/MultiServerReplication/Source/MultiServerReplication/Private/MultiServerReplicationModule.cpp
@@ -0,0 +1,6 @@
+// Copyright Epic Games, Inc. All Rights Reserved.
+
+#include "MultiServerReplicationModule.h"
+#include "Modules/ModuleManager.h"
+
+IMPLEMENT_MODULE(FMultiServerReplicationModule, MultiServerReplication);
\ No newline at end of file
diff --git a/Engine/Plugins/Runtime/MultiServerReplication/Source/MultiServerReplication/Private/MultiServerReplicationTypes.cpp b/Engine/Plugins/Runtime/MultiServerReplication/Source/MultiServerReplication/Private/MultiServerReplicationTypes.cpp
new file mode 100644
--- /dev/null
+++ b/Engine/Plugins/Runtime/MultiServerReplication/Source/MultiServerReplication/Private/MultiServerReplicationTypes.cpp
@@ -0,0 +1,5 @@
+// Copyright Epic Games, Inc. All Rights Reserved.
+
+#include "MultiServerReplicationTypes.h"
+
+DEFINE_LOG_CATEGORY(LogMultiServerReplication);
diff --git a/Engine/Plugins/Runtime/MultiServerReplication/Source/MultiServerReplication/Public/MultiServerBeaconClient.h b/Engine/Plugins/Runtime/MultiServerReplication/Source/MultiServerReplication/Public/MultiServerBeaconClient.h
new file mode 100644
--- /dev/null
+++ b/Engine/Plugins/Runtime/MultiServerReplication/Source/MultiServerReplication/Public/MultiServerBeaconClient.h
@@ -0,0 +1,84 @@
+// Copyright Epic Games, Inc. All Rights Reserved.
+
+#pragma once
+
+#include "OnlineBeaconClient.h"
+#include "MultiServerBeaconClient.generated.h"
+
+class UMultiServerNode;
+struct FUpdateLevelVisibilityLevelInfo;
+
+DECLARE_DELEGATE(FOnMultiServerConnectionEstablished);
+
+/**
+ * An online beacon that helps manage connecting to MultiServer Networks, and replicating
+ * metadata about the MultiServer Network.
+ *
+ * This Actor will exist on All MultiServer Nodes, and multiple MultiServerBeaconClient Actors may exist on non-Client Nodes,
+ * one for each other connected node.
+ */
+UCLASS(transient, config=Engine, notplaceable)
+class MULTISERVERREPLICATION_API AMultiServerBeaconClient : public AOnlineBeaconClient
+{
+	GENERATED_BODY()
+
+public:
+	AMultiServerBeaconClient();
+
+//~ Begin AOnlineBeaconClient interface
+	virtual void DestroyBeacon() override;
+	virtual void OnConnected() override;
+	virtual void OnFailure() override;
+//~ End AOnlineBeaconClient interface
+
+	/**
+	 * Attempt to connect to another specified MultiServer node.
+	 * AOnlineBeaconClient::OnFailure will be called immediately if there's a problem
+	 * within this call.
+	 *
+	 * @param ConnectionInfo	Information 
+	 */
+	virtual void ConnectToServer(const FString& ConnectInfo);
+
+	void SetOwningNode(UMultiServerNode* InOwningNode) { OwningNode = InOwningNode; }
+
+	UFUNCTION(Reliable, Server, WithValidation)
+	virtual void ServerUpdateLevelVisibility(const FUpdateLevelVisibilityLevelInfo& LevelVisibility);
+
+	UFUNCTION(Reliable, Server, WithValidation, SealedEvent)
+	void ServerUpdateMultipleLevelsVisibility(const TArray<FUpdateLevelVisibilityLevelInfo>& LevelVisibilities);
+
+	FString GetRemotePeerId() const;
+
+	/**
+	 * Actor Role on client beacons stays as authority, so this function can be used to determine whether this instance is acting as an authority.
+	 * Currently, can be used to know whether a Client or Server RPC should be called from this instance.
+	 * But once a new RPC type that supports both clients and servers simultaneously is added, this concept will be
+	 * abstracted and we shouldn't need this function anymore.
+	 */
+	bool IsAuthorityBeacon() const;
+
+protected:
+	UPROPERTY()
+	FString RemotePeerId;
+
+	virtual void OnLevelRemovedFromWorld(class ULevel* Level, class UWorld* World);
+	virtual void OnLevelAddedToWorld(class ULevel* Level, class UWorld* World);
+
+	FDelegateHandle OnLevelRemovedFromWorldHandle;
+	FDelegateHandle OnLevelAddedToWorldHandle;
+
+	FName NetworkRemapPath(FName InPackageName, bool bReading);
+
+private:
+	friend class AMultiServerBeaconHostObject;
+
+	UFUNCTION(Client, Reliable)
+	void ClientPeerConnected(const FString& NewRemotePeerId, AMultiServerBeaconClient* Beacon);
+
+	UFUNCTION(Reliable, Server)
+	void ServerSetRemotePeerId(const FString& NewRemotePeerId);
+
+	UPROPERTY(Transient)
+	TObjectPtr<UMultiServerNode> OwningNode;
+};
diff --git a/Engine/Plugins/Runtime/MultiServerReplication/Source/MultiServerReplication/Public/MultiServerNode.h b/Engine/Plugins/Runtime/MultiServerReplication/Source/MultiServerReplication/Public/MultiServerNode.h
new file mode 100644
--- /dev/null
+++ b/Engine/Plugins/Runtime/MultiServerReplication/Source/MultiServerReplication/Public/MultiServerNode.h
@@ -0,0 +1,123 @@
+// Copyright Epic Games, Inc. All Rights Reserved.
+
+#pragma once
+
+#include "UObject/Object.h"
+#include "Templates/SubclassOf.h"
+#include "MultiServerNode.generated.h"
+
+class AActor;
+class APlayerController;
+class AGameSession;
+
+class AMultiServerBeaconHost;
+class AMultiServerBeaconHostObject;
+class AMultiServerBeaconClient;
+class UMultiServerPeerConnection;
+
+DECLARE_DELEGATE_ThreeParams(FOnMultiServerConnected, const FString&, const FString&, AMultiServerBeaconClient*);
+
+/* Parameters for initializing a UMultiServerNode */
+struct FMultiServerNodeCreateParams
+{
+	/** World in which to create the node */
+	UWorld* World = nullptr;
+
+	/** String identifier of this node. Must be unique among all nodes that will connect to each other. */
+	FString LocalPeerId;
+
+	/** The port on which this node will listen for new connections. */ 
+	uint16 ListenPort = 0;
+
+	/** List of addresses of other nodes to attempt to connect to. */
+	TArray<FString> PeerAddresses;
+
+	/** Beacon client class that will be instantiated for each connection. Can implement its own RPCs. */
+	TSubclassOf<AMultiServerBeaconClient> UserBeaconClass;
+
+	/** Callback invoked when a connection to a remote node is established. */
+	FOnMultiServerConnected OnMultiServerConnected;
+};
+
+/**
+ * The MultiServer node is a system / control scheme for connecting multiple dedicated server
+ * processes to each other and allowing them to communicate via online beacons.
+ * 
+ * The basic usage pattern is to create a UMultiServerNode in project code via UMultiServerNode::Create.
+ * For a typical UE game, a good place might be in an AGameSession subclass in the RegisterServer override.
+ * The node manages all the connections to other servers, and will attempt to establish them upon creation
+ * based on the PeerAddresses in the FMultiServerNodeCreateParams.
+ * 
+ * The main user-extension point is to subclass AMultiServerBeaconClient. This subclass can implement its own
+ * RPCs to send custom messages to other servers connected to the node. The OnMultiServerConnected callback in the
+ * FMultiServerNodeCreateParams will be called when a new connection is established, with the instance of the user
+ * beacon as an argument. The user code can call its RPCs on the instance.
+ */
+UCLASS(Config=Engine, Transient, DisplayName = "MultiServer Node")
+class MULTISERVERREPLICATION_API UMultiServerNode : public UObject
+{
+	GENERATED_BODY()
+
+public:
+	UMultiServerNode();
+
+	static UMultiServerNode* Create(const FMultiServerNodeCreateParams& Params);
+
+	static void ParseCommandLineIntoCreateParams(FMultiServerNodeCreateParams& InOutParams);
+
+	bool RegisterServer(const FMultiServerNodeCreateParams& Params);
+
+	AMultiServerBeaconClient* GetBeaconClientForRemotePeer(FStringView RemotePeerId);
+
+	template<class T>
+	T* GetBeaconClientForRemotePeer(FStringView RemotePeerId);
+
+	AMultiServerBeaconClient* GetBeaconClientForURL(const FString& InURL);
+
+	template<class T>
+	T* GetBeaconClientForURL(const FString& InURL);
+
+	FString GetLocalPeerId() const { return LocalPeerId; }
+
+	float GetRetryConnectDelay() const { return RetryConnectDelay; }
+	float GetRetryConnectMaxDelay() const { return RetryConnectMaxDelay; }
+
+	TSubclassOf<AMultiServerBeaconClient> GetUserBeaconClass() const { return UserBeaconClass; }
+
+private:
+	friend AMultiServerBeaconClient;
+
+	FOnMultiServerConnected OnMultiServerConnected;
+
+	FString LocalPeerId;
+
+	UPROPERTY()
+	TObjectPtr<AMultiServerBeaconHost> BeaconHost;
+
+	UPROPERTY()
+	TObjectPtr<AMultiServerBeaconHostObject> BeaconHostObject;
+
+	UPROPERTY()
+	TArray<TObjectPtr<UMultiServerPeerConnection>> PeerConnections;
+
+	UPROPERTY(Config)
+	float RetryConnectDelay;
+
+	UPROPERTY(Config)
+	float RetryConnectMaxDelay;
+
+	TSubclassOf<AMultiServerBeaconClient> UserBeaconClass;
+};
+
+template<class T>
+T* UMultiServerNode::GetBeaconClientForRemotePeer(FStringView RemotePeerId)
+{
+	return Cast<T>(GetBeaconClientForRemotePeer(RemotePeerId));
+}
+
+template<class T>
+T* UMultiServerNode::GetBeaconClientForURL(const FString& InURL)
+
+{
+	return Cast<T>(GetBeaconClientForURL(InURL));
+}
\ No newline at end of file
diff --git a/Engine/Plugins/Runtime/MultiServerReplication/Source/MultiServerReplication/Public/MultiServerReplicationModule.h b/Engine/Plugins/Runtime/MultiServerReplication/Source/MultiServerReplication/Public/MultiServerReplicationModule.h
new file mode 100644
--- /dev/null
+++ b/Engine/Plugins/Runtime/MultiServerReplication/Source/MultiServerReplication/Public/MultiServerReplicationModule.h
@@ -0,0 +1,22 @@
+// Copyright Epic Games, Inc. All Rights Reserved.
+
+#pragma once
+
+#include "Modules/ModuleInterface.h"
+
+/** See MultiServerNode for more details about multi-server networking as a whole. */
+class FMultiServerReplicationModule : public IModuleInterface
+{
+public:
+
+	FMultiServerReplicationModule() {}
+	virtual ~FMultiServerReplicationModule() {}
+
+	// IModuleInterface
+	virtual void StartupModule() override {}
+	virtual void ShutdownModule() override {}
+	virtual bool SupportsDynamicReloading() override { return false; }
+	virtual bool SupportsAutomaticShutdown() override { return false; }
+};
+
+
diff --git a/Engine/Plugins/Runtime/MultiServerReplication/Source/MultiServerReplication/Public/MultiServerReplicationTypes.h b/Engine/Plugins/Runtime/MultiServerReplication/Source/MultiServerReplication/Public/MultiServerReplicationTypes.h
new file mode 100644
--- /dev/null
+++ b/Engine/Plugins/Runtime/MultiServerReplication/Source/MultiServerReplication/Public/MultiServerReplicationTypes.h
@@ -0,0 +1,7 @@
+// Copyright Epic Games, Inc. All Rights Reserved.
+
+#pragma once
+
+#include "Logging/LogMacros.h"
+
+MULTISERVERREPLICATION_API DECLARE_LOG_CATEGORY_EXTERN(LogMultiServerReplication, Log, All);
