Added a function to block server execution until all servers are ready.

diff --git a/Engine/Plugins/Runtime/MultiServerReplication/Source/MultiServerReplication/Private/MultiServerNode.cpp b/Engine/Plugins/Runtime/MultiServerReplication/Source/MultiServerReplication/Private/MultiServerNode.cpp
--- a/Engine/Plugins/Runtime/MultiServerReplication/Source/MultiServerReplication/Private/MultiServerNode.cpp
+++ b/Engine/Plugins/Runtime/MultiServerReplication/Source/MultiServerReplication/Private/MultiServerNode.cpp
@@ -8,6 +8,7 @@
 #include "MultiServerReplicationTypes.h"
 #include "Engine/NetConnection.h"
 #include "Engine/World.h"
+#include "Engine/NetDriver.h"
 #include "Misc/Parse.h"
 #include "Misc/CommandLine.h"
 
@@ -245,4 +246,46 @@ void UMultiServerNode::ForEachBeaconClient(TFunctionRef<void(AMultiServerBeaconC
 			Operation(BeaconClient);
 		}
 	}
+}
+
+void UMultiServerNode::ForEachNetDriver(TFunctionRef<void(UNetDriver*)> Operation)
+{
+	TSet<UNetDriver*> UniqueNetDrivers;
+
+	if (UNetDriver* HostNetDriver = BeaconHost->GetNetDriver())
+	{
+		UniqueNetDrivers.Add(HostNetDriver);
+	}
+
+	for (int32 PeerIndex = 0; PeerIndex < PeerConnections.Num(); PeerIndex++)
+	{
+		UMultiServerPeerConnection* Peer = PeerConnections[PeerIndex];
+		if (AMultiServerBeaconClient* BeaconClient = Peer->BeaconClient)
+		{
+			if (UNetConnection* Connection = BeaconClient->GetNetConnection())
+			{
+				if (UNetDriver* NetDriver = Connection->GetDriver())
+				{
+					UniqueNetDrivers.Add(NetDriver);
+				}
+			}
+		}
+	}
+
+	for (UNetDriver* NetDriver : UniqueNetDrivers)
+	{
+		Operation(NetDriver);
+	}
+}
+
+int32 UMultiServerNode::GetConnectionCount() const
+{
+	int32 ConnectionCount = 0;
+
+	const_cast<UMultiServerNode*>(this)->ForEachBeaconClient([&ConnectionCount](AMultiServerBeaconClient*)
+		{
+			ConnectionCount++;
+		});
+
+	return ConnectionCount;
 }
\ No newline at end of file
diff --git a/Engine/Plugins/Runtime/MultiServerReplication/Source/MultiServerReplication/Public/MultiServerNode.h b/Engine/Plugins/Runtime/MultiServerReplication/Source/MultiServerReplication/Public/MultiServerNode.h
--- a/Engine/Plugins/Runtime/MultiServerReplication/Source/MultiServerReplication/Public/MultiServerNode.h
+++ b/Engine/Plugins/Runtime/MultiServerReplication/Source/MultiServerReplication/Public/MultiServerNode.h
@@ -9,6 +9,7 @@
 class AActor;
 class APlayerController;
 class AGameSession;
+class UNetDriver;
 
 class AMultiServerBeaconHost;
 class AMultiServerBeaconHostObject;
@@ -86,6 +87,10 @@ public:
 
 	void ForEachBeaconClient(TFunctionRef<void(AMultiServerBeaconClient*)> Operation);
 
+	void ForEachNetDriver(TFunctionRef<void(UNetDriver*)> Operation);
+
+	int32 GetConnectionCount() const;
+
 private:
 	friend AMultiServerBeaconClient;
 
