Support for a proxy server in the MultiServerReplication plugin.

diff --git a/Engine/Plugins/Runtime/MultiServerReplication/Source/MultiServerReplication/Private/MultiServerProxy.cpp b/Engine/Plugins/Runtime/MultiServerReplication/Source/MultiServerReplication/Private/MultiServerProxy.cpp
new file mode 100644
--- /dev/null
+++ b/Engine/Plugins/Runtime/MultiServerReplication/Source/MultiServerReplication/Private/MultiServerProxy.cpp
@@ -0,0 +1,452 @@
+// Copyright Epic Games, Inc. All Rights Reserved.
+
+#include "MultiServerProxy.h"
+#include "UnrealEngine.h"
+#include "Engine/LocalPlayer.h"
+#include "Engine/GameInstance.h"
+#include "Engine/World.h"
+#include "Misc/CommandLine.h"
+#include "Net/DataChannel.h"
+
+DEFINE_LOG_CATEGORY_STATIC(LogNetProxy, Log, All);
+
+void UGameServerNotify::NotifyAcceptedConnection(class UNetConnection* Connection)
+{
+
+}
+
+EAcceptConnection::Type UGameServerNotify::NotifyAcceptingConnection()
+{
+	return EAcceptConnection::Accept;
+}
+
+bool UGameServerNotify::NotifyAcceptingChannel(class UChannel* Channel)
+{
+	return true;
+}
+
+void UGameServerNotify::NotifyControlMessage(UNetConnection* GameServerConnection, uint8 MessageType, class FInBunch& Bunch)
+{
+	Super::NotifyControlMessage(GameServerConnection, MessageType, Bunch);
+
+	// The default implementation of UPendingNetGame will only send the join request to the server when the level
+	// has loaded. Since the proxy is not currently dependent on loading levels we just send the join request 
+	// when receiving the welcome message to shortcut this logic.
+	if (MessageType == NMT_Welcome)
+	{
+		SendJoin();
+	}
+}
+
+void UProxyListenerNotify::SetProxyNetDriver(TObjectPtr<class UProxyNetDriver> InProxyNetDriver)
+{
+	ProxyNetDriver = InProxyNetDriver;
+}
+
+void UProxyListenerNotify::NotifyAcceptedConnection(class UNetConnection* Connection)
+{
+
+}
+
+EAcceptConnection::Type UProxyListenerNotify::NotifyAcceptingConnection()
+{
+	return EAcceptConnection::Accept;
+}
+
+bool UProxyListenerNotify::NotifyAcceptingChannel(class UChannel* Channel)
+{
+	return true;
+}
+
+void UProxyListenerNotify::NotifyControlMessage(UNetConnection* ProxyConnection, uint8 MessageType, class FInBunch& Bunch)
+{
+	check(ProxyNetDriver);
+
+	auto ProcessConnectionJoin = [this, ProxyConnection](int32 GameServerConnectionStateIndex, FGameServerConnectionState* GameServerConnectionState)
+	{
+		if (GameServerConnectionState->NetDriver == nullptr)
+		{
+			static int32 GameServerDriverIndex = 0;
+			FString Name = FString::Printf(TEXT("ProxyToGameServer-%d"), GameServerDriverIndex++);
+			GEngine->CreateNamedNetDriver(ProxyNetDriver->GetWorld(), *Name, "ProxyBackendNetDriver");
+
+			GameServerConnectionState->NetDriver = GEngine->FindNamedNetDriver(ProxyNetDriver->GetWorld(), *Name);
+			GameServerConnectionState->NetDriver->SetWorld(ProxyNetDriver->GetWorld());
+			Cast<UProxyBackendNetDriver>(GameServerConnectionState->NetDriver)->SetProxyNetDriver(ProxyNetDriver);
+
+			UE_LOG(LogNetProxy, Log, TEXT("Created a game server NetDriver (name=%s, port=%d)"), *GameServerConnectionState->NetDriver->GetName(), GameServerConnectionState->Port);
+		}
+
+		// Spawn a temporary, placeholder controller and add it to the world.
+		FActorSpawnParameters SpawnInfo;
+		SpawnInfo.ObjectFlags |= RF_Transient;
+		SpawnInfo.bDeferConstruction = true;
+		APlayerController* NewController = ProxyNetDriver->GetWorld()->SpawnActor<APlayerController>(APlayerController::StaticClass(), SpawnInfo);
+		ProxyNetDriver->GetWorld()->AddController(NewController);
+					
+		// Add a player to use on the game server.
+		static int32 InternalUserId = 10;
+		FPlatformUserId UserId = FPlatformUserId::CreateFromInternalId(InternalUserId++);
+		ULocalPlayer* NewPlayer = NewObject<ULocalPlayer>(GEngine, GEngine->LocalPlayerClass);
+		NewController->SetPlayer(NewPlayer);
+		ProxyNetDriver->GetWorld()->GetGameInstance()->AddLocalPlayer(NewPlayer, UserId);
+
+		UNetConnection* GameServerConnection = GameServerConnectionState->NetDriver->ServerConnection;
+		const bool bIsFirstGameServerConnection = (GameServerConnection == nullptr);
+		if (bIsFirstGameServerConnection)
+		{
+			const int32 GameServerPort = GameServerConnectionState->Port;
+			FString URLStr = FString::Format(TEXT("127.0.0.1:{0}"), { GameServerPort });
+			FURL URL(nullptr, *URLStr, TRAVEL_Absolute);
+
+			// Start the connection flow to the game server.
+			GameServerConnectionState->GameServerNotify = NewObject<UGameServerNotify>();
+			GameServerConnectionState->GameServerNotify->Initialize(URL);
+			GameServerConnectionState->GameServerNotify->InitNetDriver(GameServerConnectionState->NetDriver);
+			GameServerConnection = GameServerConnectionState->NetDriver->ServerConnection;
+
+			ProxyNetDriver->PrepareGameServerConnection(GameServerConnection);
+
+			// UNetDriver::Notify will be reset in UPendingNetGame above so it's important that we override
+			// it here again to point to the proxy.
+			GameServerConnectionState->NetDriver->Notify = GameServerConnectionState->GameServerNotify;
+
+			// Use this placeholder controller on the incoming connection to the proxy and the outgoing
+			// connection to the game server.
+			GameServerConnection->PlayerController = NewController;
+			GameServerConnection->OwningActor = NewController;
+					
+			// Register a delegate callback when a new connection is successfully established with the game server
+			// to bind the connection to the game server to the proxy connection using a common player controller.
+			GameServerConnection->GetOnConnectionOpened().AddUObject(ProxyNetDriver, &UProxyNetDriver::BindConnections);
+
+			UE_LOG(LogNetProxy, Log, TEXT("Connecting to game server: %s:%s -> %s (controller=%s player=%s"), 
+				   *ProxyConnection->GetName(),
+				   *ProxyConnection->GetDriver()->GetName(), 
+				   *GameServerConnection->GetDriver()->GetName(),
+				   *NewController->GetName(),
+				   *NewPlayer->GetName());
+		}
+		else
+		{					
+			// Send a split screen join request on the existing connection to the game server.
+			TArray<FString> Options;
+			NewPlayer->SendSplitJoin(Options, GameServerConnectionState->NetDriver);
+
+			UE_LOG(LogNetProxy, Log, TEXT("Connecting (multiplexed) to game server: %s:%s -> %s (controller=%s player=%s"), 
+				   *ProxyConnection->GetName(),
+				   *ProxyConnection->GetDriver()->GetName(), 
+				   *GameServerConnection->GetDriver()->GetName(),
+				   *NewController->GetName(),
+				   *NewPlayer->GetName());
+		}
+
+		const int32 NewPlayerIndex = GameServerConnectionState->Players.Num();
+		GameServerConnectionState->Players.Add(NewPlayer);
+		
+		FGameServerAndPlayerIndex Key; 
+		Key.GameServerConnection = GameServerConnection;
+		Key.PlayerIndex = NewPlayerIndex;
+		ProxyNetDriver->GameToProxyConnection.Add(Key, ProxyConnection);
+
+		// The player's primary game server is currently chosen as the first game server in the list
+		// of registered game servers.
+		if (ProxyNetDriver->ProxyConnectionToPrimaryPlayer.Find(ProxyConnection) == nullptr)
+		{
+			ProxyConnection->PlayerController = NewController;
+			ProxyConnection->OwningActor = NewController;
+
+			UE_LOG(LogNetProxy, Log, TEXT("Configuring primary game server %s for player (player=%s:%d proxy_connection=%s:%s)"),
+				   *GameServerConnection->GetDriver()->GetName(),
+				   *NewPlayer->GetName(),
+				   NewPlayerIndex,
+				   *ProxyConnection->GetDriver()->GetName(),
+				   *ProxyConnection->GetName());
+
+			ProxyNetDriver->ProxyPlayerToPrimaryGameServer.Add(NewPlayer, GameServerConnectionStateIndex);
+			ProxyNetDriver->ProxyConnectionToPrimaryPlayer.Add(ProxyConnection, NewPlayer);
+		}
+	};
+
+	if (MessageType == NMT_Join)
+	{
+		ProxyNetDriver->ForEachRegisteredGameServer(ProcessConnectionJoin);
+	}
+	else
+	{
+		// Forward all other connection messages onto the existing handshake logic.
+		ProxyNetDriver->GetWorld()->NotifyControlMessage(ProxyConnection, MessageType, Bunch);
+	}
+}
+
+void UProxyBackendNetDriver::SetProxyNetDriver(TObjectPtr<UProxyNetDriver> InProxyNetDriver)
+{
+	ProxyNetDriver = InProxyNetDriver;
+}
+
+bool UProxyBackendNetDriver::InitBase(bool bInitAsClient, FNetworkNotify* InNotify, const FURL& URL, bool bReuseAddressAndPort, FString& Error)
+{
+	bool bSuccess = Super::InitBase(bInitAsClient, InNotify, URL, bReuseAddressAndPort, Error);
+
+	if (bSuccess)
+	{
+		// Don't allow any RPCs received game servers to be executed on the proxy.
+		EnableExecuteRPCFunctions(false);
+	}
+
+	return bSuccess;
+}
+
+void UProxyBackendNetDriver::ForwardRemoteFunction(UObject* RootObject, UObject* SubObject, UFunction* Function, void* Parms)
+{
+	check(ProxyNetDriver);
+
+	// This function is called when the proxy receives an RPC from a game server and will only forward the function on
+	// to the frontend net driver if it is owned by a player that considers that game server the primary game server.
+
+	AActor* OwningActor = Cast<AActor>(RootObject);
+	if (!ensure(OwningActor != nullptr))
+	{
+		return;
+	}
+
+	ULocalPlayer* OwningPlayer = Cast<ULocalPlayer>(OwningActor->GetNetOwningPlayerAnyRole());
+	if (OwningPlayer == nullptr)
+	{
+		return;
+	}
+
+	int32* GameServerConnectionIndex = ProxyNetDriver->ProxyPlayerToPrimaryGameServer.Find(OwningPlayer);
+	if (GameServerConnectionIndex && *GameServerConnectionIndex < ProxyNetDriver->GameServerConnections.Num())
+	{
+		ProxyNetDriver->ProcessRemoteFunction(OwningActor, Function, Parms, static_cast<FOutParmRec*>(nullptr), static_cast<FFrame*>(nullptr), SubObject);
+	}
+}
+
+void UProxyNetDriver::RegisterGameServer(int32 Port)
+{
+	UE_LOG(LogNetProxy, Log, TEXT("Registering proxy game server on port %d"), Port);
+
+	FGameServerConnectionState GameServerConnectionState;
+	GameServerConnectionState.Port = Port;
+	GameServerConnections.Add(GameServerConnectionState);
+}
+
+bool UProxyNetDriver::IsConnectedToAllGameServers() const
+{
+	for (const FGameServerConnectionState& GameServerConnectionState : GameServerConnections)
+	{
+		if (GameServerConnectionState.NetDriver == nullptr ||
+			GameServerConnectionState.NetDriver->ServerConnection == nullptr ||
+			GameServerConnectionState.NetDriver->ServerConnection->GetConnectionState() != EConnectionState::USOCK_Open)
+		{
+			return false;
+		}
+	}
+
+	return true;
+}
+
+void UProxyNetDriver::ForEachRegisteredGameServer(TFunction<void(int32, FGameServerConnectionState*)> Func)
+{
+	for (int32 GameServerConnectionIndex = 0; GameServerConnectionIndex < GameServerConnections.Num(); GameServerConnectionIndex++)
+	{
+		Func(GameServerConnectionIndex, &(GameServerConnections.GetData()[GameServerConnectionIndex]));
+	}
+}
+
+bool UProxyNetDriver::InitBase(bool bInitAsClient, FNetworkNotify* InNotify, const FURL& URL, bool bReuseAddressAndPort, FString& Error)
+{
+	check(!bInitAsClient);
+
+	UE_LOG(LogNetProxy, Log, TEXT("Initializing ProxyNetDriver with URL %s"), *URL.ToString());
+
+	ProxyListenerNotify = NewObject<UProxyListenerNotify>(GEngine, UProxyListenerNotify::StaticClass());
+	ProxyListenerNotify->SetProxyNetDriver(this);
+
+	bool bSuccess = Super::InitBase(bInitAsClient, ProxyListenerNotify, URL, bReuseAddressAndPort, Error);
+
+	if (bSuccess)
+	{
+		// Don't allow any RPCs received from clients to be executed, only forwarded to the game server.
+		EnableExecuteRPCFunctions(false);
+	}
+
+	FString GameServerAddresses;
+	if (FParse::Value(FCommandLine::Get(), TEXT("-ProxyGameServers="), GameServerAddresses, false))
+	{
+		TArray<FString> Addresses;
+		GameServerAddresses.ParseIntoArray(Addresses, TEXT(","), true);
+
+		for (const FString& Address : Addresses)
+		{
+			FString IPAddressStr, PortStr;
+			Address.Split(TEXT(":"), &IPAddressStr, &PortStr);
+
+			int32 Port = FCString::Atoi(ToCStr(PortStr));
+			RegisterGameServer(Port);
+		}
+	}
+
+	return bSuccess;
+}
+
+bool UProxyNetDriver::InitConnect(class FNetworkNotify* InNotify, const FURL& ConnectURL, FString& Error)
+{
+	checkf(0, TEXT("UProxyNetDriver is only intended to be used to receive connections and not establish outgoing connections."));
+	return false;
+}
+
+void UProxyNetDriver::ForwardRemoteFunction(UObject* RootObject, UObject* SubObject, UFunction* Function, void* Parms)
+{
+	// This function is called when the proxy receives an RPC from a game client and will forward the function on
+	// to the owning player's primary game server.
+
+	AActor* OwningActor = Cast<AActor>(RootObject);
+	if (!ensure(OwningActor != nullptr))
+	{
+		return;
+	}
+
+	ULocalPlayer* OwningPlayer = Cast<ULocalPlayer>(OwningActor->GetNetOwningPlayerAnyRole());
+	if (OwningPlayer == nullptr)
+	{
+		return;
+	}
+
+	int32* GameServerConnectionIndex = ProxyPlayerToPrimaryGameServer.Find(OwningPlayer);
+	if (GameServerConnectionIndex && *GameServerConnectionIndex < GameServerConnections.Num())
+	{
+		FGameServerConnectionState& GameServerConnection = GameServerConnections[*GameServerConnectionIndex];
+		GameServerConnection.NetDriver->ProcessRemoteFunction(OwningActor, Function, Parms, static_cast<FOutParmRec*>(nullptr), static_cast<FFrame*>(nullptr), SubObject);
+	}
+}
+
+bool UProxyNetDriver::ShouldReplicateFunction(AActor* Actor, UFunction* Function) const
+{
+	// If any game code in the proxy world attempts to send an RPC it should not be called.
+	return false;
+}
+
+void UProxyNetDriver::NotifyActorChannelOpen(UActorChannel* Channel, AActor* Actor)
+{
+	Super::NotifyActorChannelOpen(Channel, Actor);
+
+	// The actor roles in the proxy must be replicated to the client as-is (i.e. the role on the client
+	// must be the same as the role in the proxy). Since the client will always swap roles when receiving
+	// replicated objects and the proxy is transparent to the client, the role is swapped on the proxy 
+	// before replicating.
+	SetRoleSwapOnReplicate(Actor, true);
+}
+
+void UProxyNetDriver::AddNetworkActor(AActor* Actor)
+{
+	// Ideally the proxy shouldn't spawn any actors since it's just used as a cache to pass state
+	// between game clients and game servers. For now though, actors that have the role ROLE_Authority
+	// will have replication disabled and the role set to ROLE_None. This stops them replicating to
+	// clients and disable any game actor code that only performs when the role is ROLE_Authority.
+	//
+	// It's important to note this function is called for all actors spawned on the client, both ones
+	// loaded by the proxy and those replicated from the connected game servers. It's assumed that 
+	// the actors replicated from the game servers will not have a role of ROLE_Authority and will
+	// therefore be unaffacted by this code and replicate as normal.
+	if (Actor->GetIsReplicated())
+	{
+		if (Actor->GetLocalRole() == ROLE_Authority)
+		{
+			Actor->SetReplicates(false);
+			Actor->SetRole(ROLE_None);
+		}
+	}
+
+	Super::AddNetworkActor(Actor);
+}
+
+ULocalPlayer* UProxyNetDriver::GetPlayer(UNetConnection* GameServerConnection, APlayerController* PlayerController)
+{
+	const int32 PlayerIndex = PlayerController->NetPlayerIndex;
+	ULocalPlayer* FoundPlayer = nullptr;
+
+	for (FGameServerConnectionState& GameServerConnectionState : GameServerConnections)
+	{
+		if (GameServerConnectionState.NetDriver == GameServerConnection->GetDriver())
+		{
+			if (ensure(PlayerIndex < GameServerConnectionState.Players.Num()))
+			{
+				FoundPlayer = GameServerConnectionState.Players[PlayerIndex];
+			}
+		}
+	}
+
+	return FoundPlayer;
+}
+
+void UProxyNetDriver::PrepareGameServerConnection(UNetConnection* GameServerConnection)
+{
+	// Register a callback to retrieve the correct local player for this connection based on the player controller.
+	GameServerConnection->SetLocalPlayerFunction([this](
+		UNetConnection* GameServerConnection,
+		APlayerController* PlayerController)
+	{
+		return GetPlayer(GameServerConnection, PlayerController);
+	});
+}
+
+void UProxyNetDriver::BindConnections(UNetConnection* ChildGameServerConnection,
+									  APlayerController* PlayerController,
+									  UNetConnection* ParentGameServerConnection)
+{
+	const int32 PlayerIndex = PlayerController->NetPlayerIndex;
+	FGameServerAndPlayerIndex Key; 
+	Key.GameServerConnection = ParentGameServerConnection;
+	Key.PlayerIndex = PlayerIndex;
+
+	TObjectPtr<UNetConnection>* ProxyConnectionPtr = GameToProxyConnection.Find(Key);
+	if (ProxyConnectionPtr)
+	{
+		TObjectPtr<UNetConnection> ProxyConnection = *ProxyConnectionPtr;
+		check(ProxyConnection->Driver == this);
+
+		ULocalPlayer* LocalPlayer = GetPlayer(ParentGameServerConnection, PlayerController);
+		if (ensure(LocalPlayer))
+		{
+			int32* PlayerPrimaryGameServerIdx = ProxyPlayerToPrimaryGameServer.Find(LocalPlayer);
+			if (PlayerPrimaryGameServerIdx && *PlayerPrimaryGameServerIdx < GameServerConnections.Num())
+			{
+				FGameServerConnectionState& PlayerPrimaryGameServer = GameServerConnections[*PlayerPrimaryGameServerIdx];
+				if (PlayerPrimaryGameServer.NetDriver->ServerConnection == ParentGameServerConnection)
+				{
+					UE_LOG(LogNetProxy, Log, TEXT("Binding proxy connection %s:%s to game server connection %s:%s (controller=%s player=%s:%d)"),
+						   *ProxyConnection->GetDriver()->GetName(),
+						   *ProxyConnection->GetName(),
+						   *ChildGameServerConnection->GetDriver()->GetName(),
+						   *ChildGameServerConnection->GetName(),
+						   *PlayerController->GetName(),
+						   *LocalPlayer->GetName(),
+						   PlayerIndex);
+
+					// Called here because this is the first time we have a pointer to the child connection.
+					PrepareGameServerConnection(ChildGameServerConnection);
+
+					// Ensure that the proxy connection always uses the local player associated with this game connection.
+					ProxyConnection->SetLocalPlayerFunction([LocalPlayer](
+						UNetConnection* GameServerConnection,
+						APlayerController* PlayerController)
+					{
+						return LocalPlayer;
+					});
+
+					// Complete the proxy connection handshake using the player controller replicated from the game server.
+					ProxyConnection->HandleClientPlayer(PlayerController, ProxyConnection);
+				}
+			}
+			
+			// Remap the proxy connection to the new child connection.
+			FGameServerAndPlayerIndex NewKey; 
+			NewKey.GameServerConnection = ChildGameServerConnection; 
+			NewKey.PlayerIndex = PlayerIndex;
+			GameToProxyConnection.Remove(Key);
+			GameToProxyConnection.Add(NewKey, ProxyConnection);
+		}
+	}
+}
\ No newline at end of file
diff --git a/Engine/Plugins/Runtime/MultiServerReplication/Source/MultiServerReplication/Public/MultiServerProxy.h b/Engine/Plugins/Runtime/MultiServerReplication/Source/MultiServerReplication/Public/MultiServerProxy.h
new file mode 100644
--- /dev/null
+++ b/Engine/Plugins/Runtime/MultiServerReplication/Source/MultiServerReplication/Public/MultiServerProxy.h
@@ -0,0 +1,221 @@
+// Copyright Epic Games, Inc. All Rights Reserved.
+
+#pragma once
+
+#include "Engine/PendingNetGame.h"
+#include "IpNetDriver.h"
+#include "MultiServerProxy.generated.h"
+
+/**
+ * All of the state associated with a proxy server connection to a backend game server.
+ */ 
+USTRUCT()
+struct FGameServerConnectionState
+{
+	GENERATED_BODY()
+
+	int32 Port;
+
+	UPROPERTY()
+	TObjectPtr<class UNetDriver> NetDriver;
+
+	UPROPERTY()
+	TArray<TObjectPtr<class ULocalPlayer>> Players;
+
+	UPROPERTY()
+	TObjectPtr<UGameServerNotify> GameServerNotify;
+};
+
+USTRUCT()
+struct FGameServerAndPlayerIndex
+{
+	GENERATED_BODY()
+
+	bool operator==(const FGameServerAndPlayerIndex& Rhs) const
+	{
+		return GetTypeHash(Rhs) == GetTypeHash(*this);
+	}
+
+	friend uint32 GetTypeHash(const FGameServerAndPlayerIndex& Obj)
+	{
+		return HashCombine(::GetTypeHash(Obj.PlayerIndex), ::GetTypeHash(Obj.GameServerConnection));
+	}
+
+	UPROPERTY()
+	TObjectPtr<UNetConnection> GameServerConnection;
+
+	int32 PlayerIndex;
+};
+
+/** 
+ * Intercept outgoing connection requests to game servers from the proxy server.
+ *
+ * Ensure that NMT_Join is sent after receiving NMT_Welcome from a game server. Normally
+ * NMT_Join will be sent after a level is loaded but the proxy doesn't currently handle 
+ * loading levels when connecting to a server. IMPORTANT: Once level loading functionality 
+ * has been added, this class can be removed.
+ *
+ * Defaults to the behavior in UPendingNetGame which normally handles all outgoing 
+ * connections to a game server.
+ */
+UCLASS()
+class UGameServerNotify : public UPendingNetGame
+{
+public:
+
+	GENERATED_BODY()
+
+	virtual void 					NotifyAcceptedConnection(UNetConnection* Connection) override;
+	virtual EAcceptConnection::Type NotifyAcceptingConnection() override;
+	virtual bool 					NotifyAcceptingChannel(class UChannel* Channel) override;
+	virtual void 					NotifyControlMessage(UNetConnection* GameServerConnection, uint8 MessageType, class FInBunch& Bunch) override;
+};
+
+/** 
+ * Intecept incoming from clients to the proxy server.
+ *
+ * Intercept NMT_Join when a client connects to the proxy, establishes a connection to
+ * a game server and performs the logic required to associate these two connections and
+ * forward state replicated from the game server to the client.
+ * 
+ * Defaults to the behavior in UWorld which normally handles all incoming game server 
+ * connections.
+ */
+UCLASS()
+class UProxyListenerNotify : public UObject, 
+							 public FNetworkNotify
+{
+public:
+
+	GENERATED_BODY()
+
+	void SetProxyNetDriver(TObjectPtr<class UProxyNetDriver> InProxyNetDriver);
+
+	virtual void 					NotifyAcceptedConnection(UNetConnection* Connection) override;
+	virtual EAcceptConnection::Type NotifyAcceptingConnection() override;
+	virtual bool 					NotifyAcceptingChannel(class UChannel* Channel) override;
+	virtual void 					NotifyControlMessage(UNetConnection* ProxyConnection, uint8 MessageType, class FInBunch& Bunch) override;
+
+private:
+
+	UPROPERTY()
+	TObjectPtr<class UProxyNetDriver> ProxyNetDriver;
+};
+
+/** 
+ * A driver that is used by UProxyNetDriver to connect to backend game servers.
+ */
+UCLASS()
+class MULTISERVERREPLICATION_API UProxyBackendNetDriver : public UIpNetDriver
+{
+	GENERATED_BODY()
+
+public:
+
+	void SetProxyNetDriver(TObjectPtr<UProxyNetDriver> InParentNetDriver);
+
+	virtual bool InitBase(bool bInitAsClient, FNetworkNotify* InNotify, const FURL& URL, bool bReuseAddressAndPort, FString& Error);
+	virtual void ForwardRemoteFunction(UObject* RootObject, UObject* SubObject, UFunction* Function, void* Parms) override;
+
+private:
+
+	UPROPERTY()
+	TObjectPtr<UProxyNetDriver> ProxyNetDriver;
+};
+
+/**
+* A network proxy that intercepts and forwards UE game network connections to game servers.
+*
+* This proxy externally behaves the same as a normal game server which means any game client
+* can connect to it without any special configuration.
+*
+* Internally, the proxy is made up of a NetDriver that listens for incoming connections known
+* as proxy connections, and a NetDriver for each connection to a backend game server. State from
+* the backend servers is replicated into a single, shared UWorld and the listening NetDriver 
+* replicates that state out to proxy connections.
+*
+* When a client connects to the proxy, the proxy establishes a connection to each backend game
+* server. This means that the proxy will hold an instance of ULocalPlayer and APlayerController 
+* for each client connected to the proxy multiplied by each game server (e.g. 2 clients and 2 
+* game servers would result in 4 instances of APlayerController and ULocalPlayer).
+* 
+* The instances of APlayerController set their member variable APlayerController::NetConnection
+* to a proxy connection. This means that calls to APlayerController::GetNetConnection() will 
+* always be to a proxy connection that lead to the instantiation of that controller.
+*
+* A player connected to the proxy considers one of the backend servers to be it's primary 
+* game server. This allows the proxy to preserve the idea that a player is logically connected
+* to only one game server (e.g. a player controller can only possess a single pawn at a time).
+*
+* All actors replicated to the proxy from remote game servers will have the same role as a client
+* (ROLE_SimulatedProxy or ROLE_AutonomousProxy) and will be replicated as-is to the proxy client.
+*/
+UCLASS()
+class MULTISERVERREPLICATION_API UProxyNetDriver : public UIpNetDriver
+{
+	GENERATED_BODY()
+
+public:
+	
+	/** Register a game server. */
+	void RegisterGameServer(int32 Port);
+
+	/** Return true if all registered servers are connected. */
+	bool IsConnectedToAllGameServers() const;
+
+	/** Enumerate through all outgoing connections to game servers. */
+	void ForEachRegisteredGameServer(TFunction<void(int32, FGameServerConnectionState*)> Func);
+
+	virtual bool InitBase(bool bInitAsClient, FNetworkNotify* InNotify, const FURL& URL, bool bReuseAddressAndPort, FString& Error) override;
+	virtual bool InitConnect(class FNetworkNotify* InNotify, const FURL& ConnectURL, FString& Error) override;
+	virtual void ForwardRemoteFunction(UObject* RootObject, UObject* SubObject, UFunction* Function, void* Parms) override;
+	virtual bool ShouldReplicateFunction(AActor* Actor, UFunction* Function) const override;
+	virtual void NotifyActorChannelOpen(UActorChannel* Channel, AActor* Actor) override;
+	virtual void AddNetworkActor(AActor* Actor) override;
+
+private:
+
+	friend class UProxyListenerNotify;
+
+	// Classes in EngineTest that test the behavior of UMultiServerProxy.
+	friend class FNetTestProxyRPCRouting;
+	friend class FNetTestProxyClientValidation;
+	friend class FNetTestProxyConnectionRouting;
+	
+	friend class UProxyBackendNetDriver;
+	
+	/** Return the local player associated with a connection to a game server. */
+	class ULocalPlayer* GetPlayer(UNetConnection* GameServerConnection, APlayerController* PlayerController);
+
+	/** Perform common configuration of a connection to a game server. */
+	void PrepareGameServerConnection(UNetConnection* GameServerConnection);
+
+	/** Associate a proxy connection and game server connection. */
+	void BindConnections(UNetConnection* ChildGameServerConnection,
+						 APlayerController* PlayerController,
+						 UNetConnection* ParentGameServerConnection);
+
+	/** Map a game server and player controller index to a proxy connection. */
+	UPROPERTY()
+	TMap<FGameServerAndPlayerIndex, TObjectPtr<UNetConnection>> GameToProxyConnection;
+
+	/** Map a proxy connection to a proxy player. */
+	UPROPERTY()
+	TMap<TObjectPtr<UNetConnection>, TObjectPtr<class ULocalPlayer>> ProxyConnectionToPrimaryPlayer;
+
+	/** Map a proxy player to it's primary game server index. */
+	UPROPERTY()
+	TMap<TObjectPtr<class ULocalPlayer>, int32> ProxyPlayerToPrimaryGameServer;
+
+	/** 
+	 * Net drivers and associated state used to connect to backend game servers.
+	 *
+	 * NOTE: This will need to be replaced with a TMap or similar to handle disconnects to the game servers.
+	 */
+	UPROPERTY()
+	TArray<FGameServerConnectionState> GameServerConnections;
+
+	/** Proxy listener handshake logic. */
+	UPROPERTY()
+	TObjectPtr<UProxyListenerNotify> ProxyListenerNotify;
+};
diff --git a/Engine/Source/Programs/AutomationTool/Scripts/LaunchMultiServer.Automation.cs b/Engine/Source/Programs/AutomationTool/Scripts/LaunchMultiServer.Automation.cs
--- a/Engine/Source/Programs/AutomationTool/Scripts/LaunchMultiServer.Automation.cs
+++ b/Engine/Source/Programs/AutomationTool/Scripts/LaunchMultiServer.Automation.cs
@@ -194,7 +194,9 @@ namespace AutomationTool
 		private enum ProcessLaunchType
 		{
 			Server,
-			DirectClient
+			DirectClient,
+			ProxyServer,
+			ProxyClient
 		}
 
 		public override ExitCode Execute()
@@ -233,12 +235,24 @@ namespace AutomationTool
 
 				StartProcessesForServerInfos(UniqueServerInfos, ProcessLaunchType.Server, " -MultiServerPeers=" + PeerList);
 
-				// If -directclients is specified, launch one client per server to connect directly.
-				if (ParseParam("directclients"))
+				if (ParseParam("client"))
 				{
-					Logger.LogInformation("Starting direct client instances connecting to MultiServer instances");
+					string[] ClientValues = ParseParamValue("client").Split(',');
 
-					StartProcessesForServerInfos(UniqueServerInfos, ProcessLaunchType.DirectClient, "");
+					foreach (string ClientValue in ClientValues)
+					{
+						if (ClientValue == "direct")
+						{
+							Logger.LogInformation("Starting direct client instances connecting to MultiServer instances");
+							StartProcessesForServerInfos(UniqueServerInfos, ProcessLaunchType.DirectClient, "");
+						}
+						else if (ClientValue == "proxy")
+						{
+							Logger.LogInformation("Starting proxy server connected to MultiServer instances");
+							StartProcessesForServerInfos(UniqueServerInfos, ProcessLaunchType.ProxyServer, "");
+							StartProcessesForServerInfos(UniqueServerInfos, ProcessLaunchType.ProxyClient, "");
+						}
+					}
 				}
 			}
 			else
@@ -258,11 +272,11 @@ namespace AutomationTool
 				(
 					Command: this,
 					RawProjectPath: ProjectPath,
-					DedicatedServer: LaunchType == ProcessLaunchType.Server,
-					Client: LaunchType == ProcessLaunchType.DirectClient
+					DedicatedServer: LaunchType == ProcessLaunchType.Server || LaunchType ==ProcessLaunchType.ProxyServer,
+					Client: LaunchType == ProcessLaunchType.DirectClient || LaunchType == ProcessLaunchType.ProxyClient
 				);
 
-			var DeployContexts = Project.CreateDeploymentContext(Params, LaunchType == ProcessLaunchType.Server);
+			var DeployContexts = Project.CreateDeploymentContext(Params, LaunchType == ProcessLaunchType.Server || LaunchType == ProcessLaunchType.ProxyServer);
 
 			if (DeployContexts.Count == 0)
 			{
@@ -279,7 +293,7 @@ namespace AutomationTool
 			}
 
 			var CommonArgs = (Params.Cook ? "" : DeployContext.ProjectArgForCommandLines + " ");
-			if (LaunchType == ProcessLaunchType.Server)
+			if (LaunchType == ProcessLaunchType.Server || LaunchType == ProcessLaunchType.ProxyServer)
 			{
 				CommonArgs += Params.MapToRun + " -server -log -newconsole ";
 			}
@@ -288,39 +302,97 @@ namespace AutomationTool
 
 			try
 			{
-				int ServerNum = 0;
-				foreach (var ServerInfo in ServerInfos.Values)
+				if (LaunchType == ProcessLaunchType.DirectClient || LaunchType == ProcessLaunchType.Server)
 				{
-					Logger.LogInformation("Starting MultiServer instance of type {0} for LocalId {1}, beacon port {2}, and game port {3}", LaunchType, ServerInfo.Definition.LocalId, ServerInfo.Definition.ListenPort, ServerInfo.GameListenPort);
+					int ServerNum = 0;
+					foreach (var ServerInfo in ServerInfos.Values)
+					{
+						Logger.LogInformation("Starting MultiServer instance of type {0} for LocalId {1}, beacon port {2}, and game port {3}", LaunchType, ServerInfo.Definition.LocalId, ServerInfo.Definition.ListenPort, ServerInfo.GameListenPort);
 
-					var Args = CommonArgs;
+						var Args = CommonArgs;
 
-					if (LaunchType == ProcessLaunchType.Server)
-					{
+						if (LaunchType == ProcessLaunchType.Server)
+						{
 						Args += String.Format("-MultiServerLocalId={0} -MultiServerListenPort={1} -port={2} ", ServerInfo.Definition.LocalId, ServerInfo.Definition.ListenPort, ServerInfo.GameListenPort);
 						Args += String.Format("-ConsoleTitle=\"{0} Server ID {1}\" ", Params.ShortProjectName, ServerInfo.Definition.LocalId);
 						Args += Params.ServerCommandline + " ";
+						}
+						else if (LaunchType == ProcessLaunchType.DirectClient)
+						{
+							Args += String.Format("127.0.0.1:{0} ", ServerInfo.GameListenPort) + (Params.Cook ? "" : "-game ");
+							Args += Params.ClientCommandline + " ";
+
+							// Save a separate window position for each client
+							Args += "-windowed -SaveWinPos=" + (ServerNum + 1) + " ";
+						}
+
+						if (ParseParam("notimeouts"))
+						{
+							Args += "-notimeouts ";
+						}
+
+						if (IsAttachingDebugger(LaunchType))
+						{
+							Args += "-WaitForDebuggerNoBreak ";
+						}
+
+						Args += AdditionalParameters;
+
+						var NewProcess = Run(App, Args, null, ERunOptions.Default | ERunOptions.NoWaitForExit | ERunOptions.NoStdOutRedirect);
+
+						if (NewProcess != null)
+						{
+							// Remove started process so it won't be killed on UAT exit.
+							// Essentially forces the -NoKill command-line option behavior for these.
+							ProcessManager.RemoveProcess(NewProcess);
+						}
+
+						if (IsAttachingDebugger(LaunchType))
+						{
+							ProcessDebugger.DebugProcess(NewProcess.ProcessObject.Id);
+						}
+
+						ServerNum++;
 					}
-					else if (LaunchType == ProcessLaunchType.DirectClient)
-					{
-						Args += String.Format("127.0.0.1:{0} ", ServerInfo.GameListenPort) + (Params.Cook ? "" : "-game ");
-						Args += Params.ClientCommandline + " ";
-						
-						// Save a separate window position for each client
-						Args += "-windowed -SaveWinPos=" + (ServerNum + 1) + " ";
-					}
+				}
+				else if (LaunchType == ProcessLaunchType.ProxyServer)
+				{
+					var ProxyListeningPort = 8000;
 
-					if (ParseParam("notimeouts"))
+					Logger.LogInformation("Starting MultiServer proxy listening on port {0}", ProxyListeningPort);
+					
+					var Args = CommonArgs;
+
+					Args += String.Format("-port={0} ", ProxyListeningPort);
+					Args += "-NetDriverOverrides=/Script/MultiServerReplication.ProxyNetDriver ";
+
+					if (ServerInfos.Count() > 0)
 					{
-						Args += "-notimeouts ";
+						Args += String.Format("-ConsoleTitle=\"{0} Proxy Server\" ", Params.ShortProjectName);
+						Args += "-ProxyGameServers=";
+
+						foreach (var ServerInfo in ServerInfos.Values)
+						{
+							Args += String.Format("127.0.0.1:{0},", ServerInfo.GameListenPort);
+						}
 					}
 
-					if (IsAttachingDebugger(LaunchType))
+					var NewProcess = Run(App, Args, null, ERunOptions.Default | ERunOptions.NoWaitForExit | ERunOptions.NoStdOutRedirect);
+
+					if (NewProcess != null)
 					{
-						Args += "-WaitForDebuggerNoBreak ";
+						// Remove started process so it won't be killed on UAT exit.
+						// Essentially forces the -NoKill command-line option behavior for these.
+						ProcessManager.RemoveProcess(NewProcess);
 					}
+				}
+				else if (LaunchType == ProcessLaunchType.ProxyClient)
+				{
+					var Args = CommonArgs;
 
-					Args += AdditionalParameters;
+					Args += "127.0.0.1:8000 " + (Params.Cook ? "" : "-game ");
+					Args += Params.ClientCommandline + " ";
+					Args += "-windowed";
 
 					var NewProcess = Run(App, Args, null, ERunOptions.Default | ERunOptions.NoWaitForExit | ERunOptions.NoStdOutRedirect);
 
@@ -330,13 +402,6 @@ namespace AutomationTool
 						// Essentially forces the -NoKill command-line option behavior for these.
 						ProcessManager.RemoveProcess(NewProcess);
 					}
-
-					if (IsAttachingDebugger(LaunchType))
-					{ 
-						ProcessDebugger.DebugProcess(NewProcess.ProcessObject.Id);
-					}
-
-					ServerNum++;
 				}
 			}
 			catch
diff --git a/Engine/Source/Runtime/Engine/Classes/Engine/LocalPlayer.h b/Engine/Source/Runtime/Engine/Classes/Engine/LocalPlayer.h
--- a/Engine/Source/Runtime/Engine/Classes/Engine/LocalPlayer.h
+++ b/Engine/Source/Runtime/Engine/Classes/Engine/LocalPlayer.h
@@ -455,9 +455,16 @@ public:
 	 * 	so it's only necessary to manually call this for viewports created after that
 	 * if the join fails (because the server was full, for example) all viewports on this client will be disconnected
 	 *
+	 * This call will send the command using the current net driver associated with world.
+	 *
 	 * @param	Options		array of URL options to append.
 	 */
 	ENGINE_API virtual void SendSplitJoin(TArray<FString>& Options);
+
+	/** 
+	 * Same as SendSplitJoin() above but with a additional NetDriver argument which is the drvier to send the request.
+	 */
+	ENGINE_API virtual void SendSplitJoin(TArray<FString>& Options, UNetDriver* NetDriver);
 	
 	/**
 	 * Change the physical ControllerId for this player; if the specified ControllerId is already taken by another player, changes the ControllerId
diff --git a/Engine/Source/Runtime/Engine/Classes/Engine/NetConnection.h b/Engine/Source/Runtime/Engine/Classes/Engine/NetConnection.h
--- a/Engine/Source/Runtime/Engine/Classes/Engine/NetConnection.h
+++ b/Engine/Source/Runtime/Engine/Classes/Engine/NetConnection.h
@@ -902,6 +902,32 @@ public:
 	FOnLowLevelSend			LowLevelSendDel;
 #endif
 
+	/** Set a callback that returns a local player to be used by HandleClientPlayer(). */
+	ENGINE_API void SetLocalPlayerFunction(TFunction<class ULocalPlayer*(UNetConnection*, APlayerController*)> GetPlayerFunction);
+	ENGINE_API TFunction<class ULocalPlayer*(UNetConnection*, APlayerController*)> GetLocalPlayerFunction();
+
+	/**
+	* Delegate that is called when a connection has been successfully completed on a client and
+	* calls HandleClientPlayer(). 
+	* 
+	* In the case of UChildConnection the delegate on the parent connection will be called.
+	*
+	* @param ChildGameServerConnection 		A child UNetConnection (will be equal to ParentGameServerConnection on a root connection).
+	* @param PlayerController				The controller replicated from the server.
+	* @param ParentGameServerConnection		A parent UNetConnection.
+	*/
+	DECLARE_MULTICAST_DELEGATE_ThreeParams(FOnConnectionOpened, 
+										   UNetConnection* 		/* ChildGameServerConnection */, 
+										   APlayerController* 	/* PlayerController */, 
+										   UNetConnection* 		/* ParentGameServerConnection */);
+	ENGINE_API FOnConnectionOpened& GetOnConnectionOpened() { return OnConnectionOpened; };
+
+private:
+
+	FOnConnectionOpened OnConnectionOpened;
+	TFunction<class ULocalPlayer*(UNetConnection*, APlayerController*)> PlayerConnectionFunction;
+
+public:
 
 	/**
 	 * Called to determine if a voice packet should be replicated to this
diff --git a/Engine/Source/Runtime/Engine/Classes/Engine/NetDriver.h b/Engine/Source/Runtime/Engine/Classes/Engine/NetDriver.h
--- a/Engine/Source/Runtime/Engine/Classes/Engine/NetDriver.h
+++ b/Engine/Source/Runtime/Engine/Classes/Engine/NetDriver.h
@@ -1434,6 +1434,16 @@ protected:
 	ENGINE_API virtual bool Exec_Dev(UWorld* InWorld, const TCHAR* Cmd, FOutputDevice& Ar=*GLog) override;
 //~ End FExec Interface.
 
+public:
+
+	/** Enable/disable the execution of received RPC functions. */
+	ENGINE_API void EnableExecuteRPCFunctions(bool bEnable);
+	ENGINE_API bool IsExecuteRPCFunctionsEnabled();
+
+private:
+
+	bool bExecuteRPCFunctions = true;
+
 public:
 
 	ENGINE_API ENetMode	GetNetMode() const;
@@ -1711,7 +1721,7 @@ public:
 	ENGINE_API void ForceActorRelevantNextUpdate(AActor* Actor);
 
 	/** Tells the net driver about a networked actor that was spawned */
-	ENGINE_API void AddNetworkActor(AActor* Actor);
+	ENGINE_API virtual void AddNetworkActor(AActor* Actor);
 
 	/** Called when a spawned actor is destroyed. */
 	ENGINE_API virtual void NotifyActorDestroyed( AActor* Actor, bool IsSeamlessTravel=false );
@@ -2290,9 +2300,11 @@ private:
 
 	void InitNetTraceId();
 
+protected:
 	/** Called from RPC processing code to forward RPC to other NetDrivers if ShouldForwardFunction returns true. */
-	void ForwardRemoteFunction(UObject* RootObject, UObject* SubObject, UFunction* Function, void* Parms);
+	ENGINE_API virtual void ForwardRemoteFunction(UObject* RootObject, UObject* SubObject, UFunction* Function, void* Parms);
 
+private:
 	/** Go over imported network guids and map them to the newly created object. */
 	void UpdateUnmappedObjects();
 
diff --git a/Engine/Source/Runtime/Engine/Classes/Engine/PendingNetGame.h b/Engine/Source/Runtime/Engine/Classes/Engine/PendingNetGame.h
--- a/Engine/Source/Runtime/Engine/Classes/Engine/PendingNetGame.h
+++ b/Engine/Source/Runtime/Engine/Classes/Engine/PendingNetGame.h
@@ -91,7 +91,17 @@ public:
 	// Constructor.
 	UPendingNetGame(const FObjectInitializer& ObjectInitializer = FObjectInitializer::Get());
 
-	void	InitNetDriver();
+	/** 
+	 * Inititalize the NetDriver to be used for the server connection handshake.
+	 * 
+	 * If InNetDriver is null then a new NetDriver will be created in the world using the 
+	 * NAME_GameNetDriver driver definition.
+	 * 
+	 * If InNetDriver is not null then it's expected that the NetDriver passed in has only
+	 * been created and has not begun listening for connections or connecting to a remote
+	 * server.
+	 */
+	void	InitNetDriver(UNetDriver* InNetDriver = nullptr);
 
 	/**
 	 * Begin initial handshake if needed, or call SendInitialJoin.
diff --git a/Engine/Source/Runtime/Engine/Classes/GameFramework/Actor.h b/Engine/Source/Runtime/Engine/Classes/GameFramework/Actor.h
--- a/Engine/Source/Runtime/Engine/Classes/GameFramework/Actor.h
+++ b/Engine/Source/Runtime/Engine/Classes/GameFramework/Actor.h
@@ -3089,9 +3089,12 @@ public:
 	/** Return the actor responsible for replication, if any. Typically the player controller */
 	ENGINE_API virtual const AActor* GetNetOwner() const;
 
-	/** Return the owning UPlayer (if any) of this actor. This will be a local player, a net connection, or nullptr. */
+	/** Return the owning UPlayer (if any) of this actor if it's role is ROLE_Authority. This will be a local player, a net connection, or nullptr. */
 	ENGINE_API virtual class UPlayer* GetNetOwningPlayer();
 
+	/** Return the owning UPlayer (if any) of this actor, regardless of it's current role. This will be a local player, a net connection, or nullptr. */
+	ENGINE_API virtual class UPlayer* GetNetOwningPlayerAnyRole();
+
 	/**
 	 * Get the owning connection used for communicating between client/server 
 	 * @return NetConnection to the client or server for this actor
diff --git a/Engine/Source/Runtime/Engine/Classes/GameFramework/Pawn.h b/Engine/Source/Runtime/Engine/Classes/GameFramework/Pawn.h
--- a/Engine/Source/Runtime/Engine/Classes/GameFramework/Pawn.h
+++ b/Engine/Source/Runtime/Engine/Classes/GameFramework/Pawn.h
@@ -279,6 +279,7 @@ public:
 	ENGINE_API virtual void PostInitializeComponents() override;
 	ENGINE_API virtual const AActor* GetNetOwner() const override;
 	ENGINE_API virtual UPlayer* GetNetOwningPlayer() override;
+	ENGINE_API virtual UPlayer* GetNetOwningPlayerAnyRole() override;
 	ENGINE_API virtual class UNetConnection* GetNetConnection() const override;
 	ENGINE_API virtual void PostLoad() override;
 	ENGINE_API virtual void PostRegisterAllComponents() override;
diff --git a/Engine/Source/Runtime/Engine/Classes/GameFramework/PlayerController.h b/Engine/Source/Runtime/Engine/Classes/GameFramework/PlayerController.h
--- a/Engine/Source/Runtime/Engine/Classes/GameFramework/PlayerController.h
+++ b/Engine/Source/Runtime/Engine/Classes/GameFramework/PlayerController.h
@@ -1817,6 +1817,7 @@ public:
 	ENGINE_API virtual float GetNetPriority(const FVector& ViewPos, const FVector& ViewDir, AActor* Viewer, AActor* ViewTarget, UActorChannel* InChannel, float Time, bool bLowBandwidth) override;
 	ENGINE_API virtual const AActor* GetNetOwner() const override;
 	ENGINE_API virtual class UPlayer* GetNetOwningPlayer() override;
+	ENGINE_API virtual class UPlayer* GetNetOwningPlayerAnyRole() override;
 	ENGINE_API virtual class UNetConnection* GetNetConnection() const override;
 	ENGINE_API virtual bool DestroyNetworkActorHandled() override;
 	ENGINE_API virtual void DisplayDebug(class UCanvas* Canvas, const FDebugDisplayInfo& DebugDisplay, float& YL, float& YPos) override;
diff --git a/Engine/Source/Runtime/Engine/Private/Actor.cpp b/Engine/Source/Runtime/Engine/Private/Actor.cpp
--- a/Engine/Source/Runtime/Engine/Private/Actor.cpp
+++ b/Engine/Source/Runtime/Engine/Private/Actor.cpp
@@ -1720,6 +1720,15 @@ UNetConnection* AActor::GetNetConnection() const
 	return Owner ? Owner->GetNetConnection() : nullptr;
 }
 
+UPlayer* AActor::GetNetOwningPlayerAnyRole()
+{
+	if (Owner)
+	{
+		return Owner->GetNetOwningPlayer();
+	}
+	return nullptr;
+}
+
 UPlayer* AActor::GetNetOwningPlayer()
 {
 	// We can only replicate RPCs to the owning player
diff --git a/Engine/Source/Runtime/Engine/Private/DataReplication.cpp b/Engine/Source/Runtime/Engine/Private/DataReplication.cpp
--- a/Engine/Source/Runtime/Engine/Private/DataReplication.cpp
+++ b/Engine/Source/Runtime/Engine/Private/DataReplication.cpp
@@ -1384,11 +1384,14 @@ bool FObjectReplicator::ReceivedRPC(FNetBitReader& Reader, const FReplicationFla
 			RPC_ResetLastFailedReason();
 
 			{
-				UE::Net::FScopedNetContextRPC CallingRPC;
-				UE::Net::Private::FScopedRemoteRPCMode ReceivingRemoteRPC(Function, UE::Net::Private::ERemoteFunctionMode::Receiving);
-
 				// Call the function.
-				Object->ProcessEvent(Function, Parms);
+				if (Connection->Driver->IsExecuteRPCFunctionsEnabled())
+				{
+					UE::Net::FScopedNetContextRPC CallingRPC;
+					UE::Net::Private::FScopedRemoteRPCMode ReceivingRemoteRPC(Function, UE::Net::Private::ERemoteFunctionMode::Receiving);
+
+					Object->ProcessEvent(Function, Parms);
+				}
 			}
 		}
 
diff --git a/Engine/Source/Runtime/Engine/Private/LocalPlayer.cpp b/Engine/Source/Runtime/Engine/Private/LocalPlayer.cpp
--- a/Engine/Source/Runtime/Engine/Private/LocalPlayer.cpp
+++ b/Engine/Source/Runtime/Engine/Private/LocalPlayer.cpp
@@ -322,7 +322,12 @@ void ULocalPlayer::SendSplitJoin(TArray<FString>& Options)
 		NetDriver = World->GetNetDriver();
 	}
 
-	if (World == NULL || NetDriver == NULL || NetDriver->ServerConnection == NULL || NetDriver->ServerConnection->GetConnectionState() != USOCK_Open)
+	SendSplitJoin(Options, NetDriver);
+}
+
+void ULocalPlayer::SendSplitJoin(TArray<FString>& Options, UNetDriver* NetDriver)
+{
+	if (NetDriver == NULL || NetDriver->ServerConnection == NULL || NetDriver->ServerConnection->GetConnectionState() != USOCK_Open)
 	{
 		UE_LOG(LogPlayerManagement, Warning, TEXT("SendSplitJoin(): Not connected to a server"));
 	}
diff --git a/Engine/Source/Runtime/Engine/Private/NetConnection.cpp b/Engine/Source/Runtime/Engine/Private/NetConnection.cpp
--- a/Engine/Source/Runtime/Engine/Private/NetConnection.cpp
+++ b/Engine/Source/Runtime/Engine/Private/NetConnection.cpp
@@ -4911,10 +4911,17 @@ void UNetConnection::HandleClientPlayer( APlayerController *PC, UNetConnection*
 
 	// Hook up the Viewport to the new player actor.
 	ULocalPlayer*	LocalPlayer = NULL;
-	if (FLocalPlayerIterator It(GEngine, Driver->GetWorld()); It)
+	if (GetLocalPlayerFunction())
 	{
-		LocalPlayer = *It;
-	} 
+		LocalPlayer = GetLocalPlayerFunction()(NetConnection, PC);
+	}
+	else
+	{
+		if (FLocalPlayerIterator It(GEngine, Driver->GetWorld()); It)
+		{
+			LocalPlayer = *It;
+		}
+	}
 
 	// Detach old player if it's in the same level.
 	check(LocalPlayer);
@@ -4994,20 +5001,29 @@ void UNetConnection::HandleClientPlayer( APlayerController *PC, UNetConnection*
 		}
 	}
 
+	NetConnection->GetOnConnectionOpened().Broadcast(this, PC, NetConnection);
+	
 	NotifyConnectionUpdated();
 }
 
 void UChildConnection::HandleClientPlayer(APlayerController* PC, UNetConnection* NetConnection)
 {
-	// find the first player that doesn't already have a connection
 	ULocalPlayer* NewPlayer = NULL;
-	uint8 CurrentIndex = 0;
-	for (FLocalPlayerIterator It(GEngine, Driver->GetWorld()); It; ++It, CurrentIndex++)
+	if (NetConnection->GetLocalPlayerFunction())
+	{
+		NewPlayer = NetConnection->GetLocalPlayerFunction()(NetConnection, PC);
+	}
+	else
 	{
-		if (CurrentIndex == PC->NetPlayerIndex)
+		// find the first player that doesn't already have a connection
+		uint8 CurrentIndex = 0;
+		for (FLocalPlayerIterator It(GEngine, Driver->GetWorld()); It; ++It, CurrentIndex++)
 		{
-			NewPlayer = *It;
-			break;
+			if (CurrentIndex == PC->NetPlayerIndex)
+			{
+				NewPlayer = *It;
+				break;
+			}
 		}
 	}
 
@@ -5064,6 +5080,8 @@ void UChildConnection::HandleClientPlayer(APlayerController* PC, UNetConnection*
 	PlayerController = PC;
 	OwningActor = PC;
 
+	NetConnection->GetOnConnectionOpened().Broadcast(this, PC, NetConnection);
+
 	NotifyConnectionUpdated();
 }
 
@@ -5075,6 +5093,16 @@ void UNetConnection::UpdatePacketSimulationSettings(void)
 }
 #endif
 
+void UNetConnection::SetLocalPlayerFunction(TFunction<ULocalPlayer*(UNetConnection*, APlayerController*)> InPlayerConnectionFunction)
+{
+	PlayerConnectionFunction = InPlayerConnectionFunction;
+}
+
+TFunction<ULocalPlayer*(UNetConnection*, APlayerController*)> UNetConnection::GetLocalPlayerFunction()
+{
+	return PlayerConnectionFunction;
+}
+
 /**
  * Called to determine if a voice packet should be replicated to this
  * connection or any of its child connections
diff --git a/Engine/Source/Runtime/Engine/Private/NetDriver.cpp b/Engine/Source/Runtime/Engine/Private/NetDriver.cpp
--- a/Engine/Source/Runtime/Engine/Private/NetDriver.cpp
+++ b/Engine/Source/Runtime/Engine/Private/NetDriver.cpp
@@ -2062,6 +2062,16 @@ void UNetDriver::FlushHandler()
 	}
 }
 
+void UNetDriver::EnableExecuteRPCFunctions(bool bEnable)
+{
+	bExecuteRPCFunctions = bEnable;
+}
+
+bool UNetDriver::IsExecuteRPCFunctionsEnabled()
+{
+	return bExecuteRPCFunctions;
+}
+
 ENetMode UNetDriver::GetNetMode() const
 {
 	// Special case for PIE and game worlds - forcing dedicated server behavior if requested
@@ -2757,6 +2767,9 @@ void UNetDriver::ProcessRemoteFunctionForChannelPrivate(
 	// This function should be kept fast! Assume this is getting called multiple times at once. Don't look things up/recalc them if they do not change per connection/actor.
 	// ---------------------------------------------------------------------------------------------------------------------------------------------------------------------------------------------------------------------------
 
+	// The connection must be part of the same driver.
+	ensure(Connection->GetDriver() == this);
+
 	if (Ch->Closing)
 	{
 		return;
@@ -7785,6 +7798,10 @@ void UNetDriver::ProcessRemoteFunction(
 		Connection = Actor->GetNetConnection();
 		if (Connection)
 		{
+			if (ServerConnection)
+			{
+				Connection = ServerConnection;
+			}
 			InternalProcessRemoteFunction(Actor, SubObject, Connection, Function, Parameters, OutParms, Stack, bIsServer);
 		}
 		else
diff --git a/Engine/Source/Runtime/Engine/Private/Pawn.cpp b/Engine/Source/Runtime/Engine/Private/Pawn.cpp
--- a/Engine/Source/Runtime/Engine/Private/Pawn.cpp
+++ b/Engine/Source/Runtime/Engine/Private/Pawn.cpp
@@ -755,6 +755,16 @@ class UPlayer* APawn::GetNetOwningPlayer()
 	return Super::GetNetOwningPlayer();
 }
 
+class UPlayer* APawn::GetNetOwningPlayerAnyRole()
+{
+	if (Controller)
+	{
+		APlayerController* PC = Cast<APlayerController>(Controller);
+		return PC ? PC->Player : nullptr;
+	}
+
+	return Super::GetNetOwningPlayer();
+}
 
 UInputComponent* APawn::CreatePlayerInputComponent()
 {
diff --git a/Engine/Source/Runtime/Engine/Private/PendingNetGame.cpp b/Engine/Source/Runtime/Engine/Private/PendingNetGame.cpp
--- a/Engine/Source/Runtime/Engine/Private/PendingNetGame.cpp
+++ b/Engine/Source/Runtime/Engine/Private/PendingNetGame.cpp
@@ -31,7 +31,7 @@ UPendingNetGame::UPendingNetGame(const FObjectInitializer& ObjectInitializer)
 {
 }
 
-void UPendingNetGame::InitNetDriver()
+void UPendingNetGame::InitNetDriver(UNetDriver* InNetDriver)
 {
 	LLM_SCOPE(ELLMTag::Networking);
 
@@ -40,9 +40,18 @@ void UPendingNetGame::InitNetDriver()
 		NETWORK_PROFILER(GNetworkProfiler.TrackSessionChange(true, URL));
 
 		// Try to create network driver.
-		if (GEngine->CreateNamedNetDriver(this, NAME_PendingNetDriver, NAME_GameNetDriver))
+		if (InNetDriver == nullptr)
 		{
-			NetDriver = GEngine->FindNamedNetDriver(this, NAME_PendingNetDriver);
+			if (GEngine->CreateNamedNetDriver(this, NAME_PendingNetDriver, NAME_GameNetDriver))
+			{
+				NetDriver = GEngine->FindNamedNetDriver(this, NAME_PendingNetDriver);
+			}
+		}
+		else
+		{
+			ensureMsgf(InNetDriver->ServerConnection == nullptr, TEXT("UPendingNetGame: Provided NetDriver is already connected/connecting to a remote server."));
+			ensureMsgf(!InNetDriver->ConnectionlessHandler.IsValid(), TEXT("UPendingNetGame: Provided NetDriver is already listening for incoming connections."));
+			NetDriver = InNetDriver;
 		}
 
 		if (NetDriver == nullptr)
diff --git a/Engine/Source/Runtime/Engine/Private/PlayerController.cpp b/Engine/Source/Runtime/Engine/Private/PlayerController.cpp
--- a/Engine/Source/Runtime/Engine/Private/PlayerController.cpp
+++ b/Engine/Source/Runtime/Engine/Private/PlayerController.cpp
@@ -285,6 +285,11 @@ UPlayer* APlayerController::GetNetOwningPlayer()
 	return Player;
 }
 
+class UPlayer* APlayerController::GetNetOwningPlayerAnyRole()
+{
+	return Player;
+}
+
 bool APlayerController::HasNetOwner() const
 {
 	// Player controllers are their own net owners
@@ -1388,7 +1393,11 @@ void APlayerController::OnActorChannelOpen(FInBunch& InBunch, UNetConnection* Co
 	InBunch << NetPlayerIndex;
 	if (Connection->Driver != NULL && Connection == Connection->Driver->ServerConnection)
 	{
-		if (NetPlayerIndex == 0)
+		const bool bIsParentConnection = 
+			(NetPlayerIndex == 0) ||
+			(GetGameInstance() && GetGameInstance()->GetNumLocalPlayers() == 1);
+
+		if (bIsParentConnection)
 		{
 			// main connection PlayerController
 			Connection->HandleClientPlayer(this, Connection); 
diff --git a/Engine/Source/Runtime/Engine/Private/Tests/NetTestHelpers.cpp b/Engine/Source/Runtime/Engine/Private/Tests/NetTestHelpers.cpp
--- a/Engine/Source/Runtime/Engine/Private/Tests/NetTestHelpers.cpp
+++ b/Engine/Source/Runtime/Engine/Private/Tests/NetTestHelpers.cpp
@@ -88,6 +88,27 @@ FTestWorldInstance FTestWorldInstance::CreateClient(int32 ServerPort)
 	return NewInstance;
 }
 
+FTestWorldInstance FTestWorldInstance::CreateProxy()
+{
+	FGameInstancePIEParameters ServerParams;
+	ServerParams.bSimulateInEditor = false;
+	ServerParams.bAnyBlueprintErrors = false;
+	ServerParams.bStartInSpectatorMode = false;
+	ServerParams.bRunAsDedicated = true;
+	ServerParams.bIsPrimaryPIEClient = false;
+	ServerParams.WorldFeatureLevel = ERHIFeatureLevel::SM5;
+	ServerParams.EditorPlaySettings = nullptr;
+	ServerParams.NetMode = PIE_ListenServer;
+		
+	FTestWorldInstance NewInstance(ServerParams);
+
+	FURL LocalURL(nullptr, TEXT("/Engine/Maps/Entry?listen?NetDriverDef=ProxyNetDriver"), TRAVEL_Absolute);
+	FString BrowseError;
+	GEngine->Browse(*NewInstance.GameInstance->GetWorldContext(), LocalURL, BrowseError);
+
+	return NewInstance;
+}
+
 FTestWorldInstance::FTestWorldInstance(const FGameInstancePIEParameters& InstanceParams)
 {
 	GameInstance = NewObject<UGameInstance>(GEngine, UGameInstance::StaticClass());
diff --git a/Engine/Source/Runtime/Engine/Private/World.cpp b/Engine/Source/Runtime/Engine/Private/World.cpp
--- a/Engine/Source/Runtime/Engine/Private/World.cpp
+++ b/Engine/Source/Runtime/Engine/Private/World.cpp
@@ -7087,7 +7087,13 @@ bool UWorld::Listen( FURL& InURL )
 	}
 
 	// Create net driver.
-	if (GEngine->CreateNamedNetDriver(this, NAME_GameNetDriver, NAME_GameNetDriver))
+	FName NetDriverDefinition = NAME_GameNetDriver; 
+	if (InURL.HasOption(TEXT("NetDriverDef=")))
+	{
+		NetDriverDefinition = InURL.GetOption(TEXT("NetDriverDef="), nullptr);
+	}
+
+	if (GEngine->CreateNamedNetDriver(this, NAME_GameNetDriver, NetDriverDefinition))
 	{
 		NetDriver = GEngine->FindNamedNetDriver(this, NAME_GameNetDriver);
 		NetDriver->SetWorld(this);
diff --git a/Engine/Source/Runtime/Engine/Public/Tests/NetTestHelpers.h b/Engine/Source/Runtime/Engine/Public/Tests/NetTestHelpers.h
--- a/Engine/Source/Runtime/Engine/Public/Tests/NetTestHelpers.h
+++ b/Engine/Source/Runtime/Engine/Public/Tests/NetTestHelpers.h
@@ -96,6 +96,7 @@ struct FTestWorldInstance
 public:
 	ENGINE_API static FTestWorldInstance CreateServer(const TCHAR* InURL);
 	ENGINE_API static FTestWorldInstance CreateClient(int32 ServerPort);
+	ENGINE_API static FTestWorldInstance CreateProxy();
 
 	ENGINE_API ~FTestWorldInstance();
 
