Support for multiple DSTM clients connecting to the proxy.

diff --git a/Engine/Plugins/Runtime/MultiServerReplication/Source/MultiServerReplication/Private/MultiServerProxy.cpp b/Engine/Plugins/Runtime/MultiServerReplication/Source/MultiServerReplication/Private/MultiServerProxy.cpp
--- a/Engine/Plugins/Runtime/MultiServerReplication/Source/MultiServerReplication/Private/MultiServerProxy.cpp
+++ b/Engine/Plugins/Runtime/MultiServerReplication/Source/MultiServerReplication/Private/MultiServerProxy.cpp
@@ -91,6 +91,10 @@ void UProxyListenerNotify::NotifyControlMessage(UNetConnection* ProxyConnection,
 		NewController->SetPlayer(NewPlayer);
 		ProxyNetDriver->GetWorld()->GetGameInstance()->AddLocalPlayer(NewPlayer, UserId);
 
+		// The player's primary game server is currently chosen as the first game server in the list
+		// of registered game servers.
+		bool bIsPrimaryGameServer = (ProxyNetDriver->ProxyConnectionToPrimaryPlayer.Find(ProxyConnection) == nullptr);
+
 		UNetConnection* GameServerConnection = GameServerConnectionState->NetDriver->ServerConnection;
 		const bool bIsFirstGameServerConnection = (GameServerConnection == nullptr);
 		if (bIsFirstGameServerConnection)
@@ -99,6 +103,13 @@ void UProxyListenerNotify::NotifyControlMessage(UNetConnection* ProxyConnection,
 			FString URLStr = FString::Format(TEXT("127.0.0.1:{0}"), { GameServerPort });
 			FURL URL(nullptr, *URLStr, TRAVEL_Absolute);
 
+			URL.AddOption(TEXT("AutonomousAsSimulated"));
+
+			if (!bIsPrimaryGameServer)
+			{
+				URL.AddOption(TEXT("SpectatorOnly=1"));
+			}
+
 			// Start the connection flow to the game server.
 			GameServerConnectionState->GameServerNotify = NewObject<UGameServerNotify>();
 			GameServerConnectionState->GameServerNotify->Initialize(URL);
@@ -118,7 +129,7 @@ void UProxyListenerNotify::NotifyControlMessage(UNetConnection* ProxyConnection,
 					
 			// Register a delegate callback when a new connection is successfully established with the game server
 			// to bind the connection to the game server to the proxy connection using a common player controller.
-			GameServerConnection->GetOnConnectionOpened().AddUObject(ProxyNetDriver, &UProxyNetDriver::BindConnections);
+			GameServerConnection->GetOnConnectionOpened().AddUObject(ProxyNetDriver, &UProxyNetDriver::ConnectionHandshakeComplete);
 
 			UE_LOG(LogNetProxy, Log, TEXT("Connecting to game server: %s:%s -> %s (controller=%s player=%s"), 
 				   *ProxyConnection->GetName(),
@@ -128,17 +139,37 @@ void UProxyListenerNotify::NotifyControlMessage(UNetConnection* ProxyConnection,
 				   *NewPlayer->GetName());
 		}
 		else
-		{					
-			// Send a split screen join request on the existing connection to the game server.
-			TArray<FString> Options;
-			NewPlayer->SendSplitJoin(Options, GameServerConnectionState->NetDriver);
-
-			UE_LOG(LogNetProxy, Log, TEXT("Connecting (multiplexed) to game server: %s:%s -> %s (controller=%s player=%s"), 
-				   *ProxyConnection->GetName(),
-				   *ProxyConnection->GetDriver()->GetName(), 
-				   *GameServerConnection->GetDriver()->GetName(),
-				   *NewController->GetName(),
-				   *NewPlayer->GetName());
+		{
+			FGameServerSplitJoinRequest Request;
+			Request.Player = NewPlayer;
+			Request.bIsPrimaryGameServer = bIsPrimaryGameServer;
+
+			// ULocalPlayer::SendSplitJoin() can only be called when the parent connection has been successfully opened. If 
+			// the parent connection is not open yet, this player will be queued up and processed when the parent connection
+			// has been opened (see UMultiProxyServer::ConnectionHandshakeComplete()).
+			if (GameServerConnectionState->NetDriver->ServerConnection->GetConnectionState() == USOCK_Open)
+			{
+				GameServerConnectionState->PendingSplitJoinRequests.Add(Request);
+				ProxyNetDriver->FlushSplitJoinRequests(GameServerConnectionState);
+
+				UE_LOG(LogNetProxy, Log, TEXT("Connecting (multiplexed) to game server: %s:%s -> %s (controller=%s player=%s"), 
+					   *ProxyConnection->GetName(),
+					   *ProxyConnection->GetDriver()->GetName(), 
+					   *GameServerConnection->GetDriver()->GetName(),
+					   *NewController->GetName(),
+					   *NewPlayer->GetName());
+			}
+			else
+			{
+				GameServerConnectionState->PendingSplitJoinRequests.Add(Request);
+	
+				UE_LOG(LogNetProxy, Log, TEXT("Queuing up connection (multiplexed) to game server: %s:%s -> %s (controller=%s player=%s"), 
+					   *ProxyConnection->GetName(),
+					   *ProxyConnection->GetDriver()->GetName(), 
+					   *GameServerConnection->GetDriver()->GetName(),
+					   *NewController->GetName(),
+					   *NewPlayer->GetName());
+			}
 		}
 
 		const int32 NewPlayerIndex = GameServerConnectionState->Players.Num();
@@ -149,9 +180,7 @@ void UProxyListenerNotify::NotifyControlMessage(UNetConnection* ProxyConnection,
 		Key.PlayerIndex = NewPlayerIndex;
 		ProxyNetDriver->GameToProxyConnection.Add(Key, ProxyConnection);
 
-		// The player's primary game server is currently chosen as the first game server in the list
-		// of registered game servers.
-		if (ProxyNetDriver->ProxyConnectionToPrimaryPlayer.Find(ProxyConnection) == nullptr)
+		if (bIsPrimaryGameServer)
 		{
 			ProxyConnection->PlayerController = NewController;
 			ProxyConnection->OwningActor = NewController;
@@ -392,9 +421,9 @@ void UProxyNetDriver::PrepareGameServerConnection(UNetConnection* GameServerConn
 	});
 }
 
-void UProxyNetDriver::BindConnections(UNetConnection* ChildGameServerConnection,
-									  APlayerController* PlayerController,
-									  UNetConnection* ParentGameServerConnection)
+void UProxyNetDriver::ConnectionHandshakeComplete(UNetConnection* ChildGameServerConnection,
+												  APlayerController* PlayerController,
+												  UNetConnection* ParentGameServerConnection)
 {
 	const int32 PlayerIndex = PlayerController->NetPlayerIndex;
 	FGameServerAndPlayerIndex Key; 
@@ -449,4 +478,52 @@ void UProxyNetDriver::BindConnections(UNetConnection* ChildGameServerConnection,
 			GameToProxyConnection.Add(NewKey, ProxyConnection);
 		}
 	}
+
+	// If this was a parent connection send any join requests for any pending multiplexed connections.
+	if (ChildGameServerConnection == ParentGameServerConnection)
+	{
+		FGameServerConnectionState* CurrGameServerConnectionState = nullptr;
+		for (FGameServerConnectionState& GameServerConnectionState : GameServerConnections)
+		{
+			if (GameServerConnectionState.NetDriver->ServerConnection == ParentGameServerConnection)
+			{
+				CurrGameServerConnectionState = &GameServerConnectionState;
+				break;
+			}
+		}
+
+		if (ensure(CurrGameServerConnectionState))
+		{
+			FlushSplitJoinRequests(CurrGameServerConnectionState);
+		}
+	}
+}
+
+void UProxyNetDriver::FlushSplitJoinRequests(FGameServerConnectionState* GameServerConnectionState)
+{
+	if (GameServerConnectionState->NetDriver->ServerConnection->GetConnectionState() != EConnectionState::USOCK_Open)
+	{
+		UE_LOG(LogNetProxy, Error, TEXT("Flushing split join requests on %s without the parent connection being opened."), 
+			   *GameServerConnectionState->NetDriver->GetName());
+
+		return;
+	}
+
+	for (FGameServerSplitJoinRequest& Request : GameServerConnectionState->PendingSplitJoinRequests)
+	{
+		TArray<FString> Options;
+
+		if (!Request.bIsPrimaryGameServer)
+		{
+			Options.Add(TEXT("SpectatorOnly=1"));
+		}
+
+		Request.Player->SendSplitJoin(Options, GameServerConnectionState->NetDriver);
+
+		UE_LOG(LogNetProxy, Log, TEXT("Sending queued connection (multiplexed) request to game server: %s (player=%s"),
+			   *GameServerConnectionState->NetDriver->GetName(),
+			   *Request.Player->GetName());
+	}
+
+	GameServerConnectionState->PendingSplitJoinRequests.Reset();
 }
\ No newline at end of file
diff --git a/Engine/Plugins/Runtime/MultiServerReplication/Source/MultiServerReplication/Public/MultiServerProxy.h b/Engine/Plugins/Runtime/MultiServerReplication/Source/MultiServerReplication/Public/MultiServerProxy.h
--- a/Engine/Plugins/Runtime/MultiServerReplication/Source/MultiServerReplication/Public/MultiServerProxy.h
+++ b/Engine/Plugins/Runtime/MultiServerReplication/Source/MultiServerReplication/Public/MultiServerProxy.h
@@ -6,6 +6,17 @@
 #include "IpNetDriver.h"
 #include "MultiServerProxy.generated.h"
 
+USTRUCT()
+struct FGameServerSplitJoinRequest
+{
+	GENERATED_BODY()
+
+	UPROPERTY()
+	TObjectPtr<class ULocalPlayer> Player;
+
+	bool bIsPrimaryGameServer;
+};
+
 /**
  * All of the state associated with a proxy server connection to a backend game server.
  */ 
@@ -24,6 +35,9 @@ struct FGameServerConnectionState
 
 	UPROPERTY()
 	TObjectPtr<UGameServerNotify> GameServerNotify;
+
+	UPROPERTY()
+	TArray<FGameServerSplitJoinRequest> PendingSplitJoinRequests;
 };
 
 USTRUCT()
@@ -190,10 +204,13 @@ private:
 	/** Perform common configuration of a connection to a game server. */
 	void PrepareGameServerConnection(UNetConnection* GameServerConnection);
 
-	/** Associate a proxy connection and game server connection. */
-	void BindConnections(UNetConnection* ChildGameServerConnection,
-						 APlayerController* PlayerController,
-						 UNetConnection* ParentGameServerConnection);
+	/** Called when a game server connection has been established. */
+	void ConnectionHandshakeComplete(UNetConnection* ChildGameServerConnection,
+									 APlayerController* PlayerController,
+									 UNetConnection* ParentGameServerConnection);
+
+	/** Send any split join requests that have been queued up to the game server. */
+	void FlushSplitJoinRequests(FGameServerConnectionState* GameServerConnectionState);
 
 	/** Map a game server and player controller index to a proxy connection. */
 	UPROPERTY()
diff --git a/Engine/Source/Programs/AutomationTool/Scripts/LaunchMultiServer.Automation.cs b/Engine/Source/Programs/AutomationTool/Scripts/LaunchMultiServer.Automation.cs
--- a/Engine/Source/Programs/AutomationTool/Scripts/LaunchMultiServer.Automation.cs
+++ b/Engine/Source/Programs/AutomationTool/Scripts/LaunchMultiServer.Automation.cs
@@ -248,9 +248,20 @@ namespace AutomationTool
 						}
 						else if (ClientValue == "proxy")
 						{
+							int ClientCount = 1;
+							if (ParseParam("proxyclientcount"))
+							{
+								string ClientCountStr = ParseParamValue("proxyclientcount", null);
+								Int32.TryParse(ClientCountStr, out ClientCount);
+							}
+
 							Logger.LogInformation("Starting proxy server connected to MultiServer instances");
 							StartProcessesForServerInfos(UniqueServerInfos, ProcessLaunchType.ProxyServer, "");
-							StartProcessesForServerInfos(UniqueServerInfos, ProcessLaunchType.ProxyClient, "");
+
+							for (int i = 0; i < ClientCount; i++)
+							{
+								StartProcessesForServerInfos(UniqueServerInfos, ProcessLaunchType.ProxyClient, "");
+							}
 						}
 					}
 				}
diff --git a/Engine/Source/Runtime/Engine/Classes/Engine/NetDriver.h b/Engine/Source/Runtime/Engine/Classes/Engine/NetDriver.h
--- a/Engine/Source/Runtime/Engine/Classes/Engine/NetDriver.h
+++ b/Engine/Source/Runtime/Engine/Classes/Engine/NetDriver.h
@@ -1438,11 +1438,16 @@ public:
 
 	/** Enable/disable the execution of received RPC functions. */
 	ENGINE_API void EnableExecuteRPCFunctions(bool bEnable);
-	ENGINE_API bool IsExecuteRPCFunctionsEnabled();
+	ENGINE_API bool IsExecuteRPCFunctionsEnabled() const;
+
+	/** Treat actors with ROLE_AutonomousProxy as ROLE_SimulatedProxy when replicating state (applies to all client connections). */
+	ENGINE_API void ReplicateAutonomousAsSimulated(bool bEnabled);
+	ENGINE_API bool IsReplicatingAutonomousAsSimulated() const;
 
 private:
 
 	bool bExecuteRPCFunctions = true;
+	bool bReplicatingAutonomousAsSimulated = false;
 
 public:
 
diff --git a/Engine/Source/Runtime/Engine/Private/DataChannel.cpp b/Engine/Source/Runtime/Engine/Private/DataChannel.cpp
--- a/Engine/Source/Runtime/Engine/Private/DataChannel.cpp
+++ b/Engine/Source/Runtime/Engine/Private/DataChannel.cpp
@@ -3742,6 +3742,12 @@ int64 UActorChannel::ReplicateActor()
 	FScopedRoleDowngrade ScopedRoleDowngrade( Actor, RepFlags );
 
 	RepFlags.bNetSimulated	= (Actor->GetRemoteRole() == ROLE_SimulatedProxy);
+
+	if (Actor->GetRemoteRole() == ROLE_AutonomousProxy && Connection->GetDriver()->IsReplicatingAutonomousAsSimulated())
+	{
+		RepFlags.bNetSimulated = true;
+	}
+
 	RepFlags.bRepPhysics	= Actor->GetReplicatedMovement().bRepPhysics;
 	RepFlags.bReplay		= bReplay;
 	RepFlags.bClientReplay	= ActorWorld->IsRecordingClientReplay();
diff --git a/Engine/Source/Runtime/Engine/Private/NetDriver.cpp b/Engine/Source/Runtime/Engine/Private/NetDriver.cpp
--- a/Engine/Source/Runtime/Engine/Private/NetDriver.cpp
+++ b/Engine/Source/Runtime/Engine/Private/NetDriver.cpp
@@ -2067,11 +2067,21 @@ void UNetDriver::EnableExecuteRPCFunctions(bool bEnable)
 	bExecuteRPCFunctions = bEnable;
 }
 
-bool UNetDriver::IsExecuteRPCFunctionsEnabled()
+bool UNetDriver::IsExecuteRPCFunctionsEnabled() const
 {
 	return bExecuteRPCFunctions;
 }
 
+void UNetDriver::ReplicateAutonomousAsSimulated(bool bEnabled)
+{
+	bReplicatingAutonomousAsSimulated = bEnabled;
+}
+
+bool UNetDriver::IsReplicatingAutonomousAsSimulated() const
+{
+	return bReplicatingAutonomousAsSimulated;
+}
+
 ENetMode UNetDriver::GetNetMode() const
 {
 	// Special case for PIE and game worlds - forcing dedicated server behavior if requested
diff --git a/Engine/Source/Runtime/Engine/Private/World.cpp b/Engine/Source/Runtime/Engine/Private/World.cpp
--- a/Engine/Source/Runtime/Engine/Private/World.cpp
+++ b/Engine/Source/Runtime/Engine/Private/World.cpp
@@ -6809,6 +6809,11 @@ void UWorld::NotifyControlMessage(UNetConnection* Connection, uint8 MessageType,
 						break;
 					}
 
+					if (InURL.HasOption(TEXT("AutonomousAsSimulated")))
+					{
+						NetDriver->ReplicateAutonomousAsSimulated(true);
+					}
+
 					int32 SplitscreenCount = FMath::Min(Connection->Children.Num() + 1, 255);
 
 					// Don't allow clients to specify this value
