Disable replication transactionality on the multi server proxy.

diff --git a/Engine/Plugins/Runtime/MultiServerReplication/Source/MultiServerReplication/Private/MultiServerProxy.cpp b/Engine/Plugins/Runtime/MultiServerReplication/Source/MultiServerReplication/Private/MultiServerProxy.cpp
--- a/Engine/Plugins/Runtime/MultiServerReplication/Source/MultiServerReplication/Private/MultiServerProxy.cpp
+++ b/Engine/Plugins/Runtime/MultiServerReplication/Source/MultiServerReplication/Private/MultiServerProxy.cpp
@@ -223,6 +223,8 @@ bool UProxyBackendNetDriver::InitBase(bool bInitAsClient, FNetworkNotify* InNoti
 		EnableExecuteRPCFunctions(false);
 	}
 
+	SetReplicateTransactionally(false);
+
 	return bSuccess;
 }
 
@@ -317,6 +319,8 @@ bool UProxyNetDriver::InitBase(bool bInitAsClient, FNetworkNotify* InNotify, con
 		}
 	}
 
+	SetReplicateTransactionally(false);
+
 	return bSuccess;
 }
 
