Share the Root Playspace by using newly implemented BroadcastActorToAllPeers_Slow.  That function will ensure that an Actor gets passed around to all servers, so that they are all aware of it.  The second step is that the new [project specific] RootPlayspace knows when it Migrates to set itself as the Root Playspace and invalidate any pointers we're not ready to migrate yet.

diff --git a/Engine/Plugins/Runtime/MultiServerReplication/Source/MultiServerReplication/Private/MultiServerBeaconClient.cpp b/Engine/Plugins/Runtime/MultiServerReplication/Source/MultiServerReplication/Private/MultiServerBeaconClient.cpp
--- a/Engine/Plugins/Runtime/MultiServerReplication/Source/MultiServerReplication/Private/MultiServerBeaconClient.cpp
+++ b/Engine/Plugins/Runtime/MultiServerReplication/Source/MultiServerReplication/Private/MultiServerBeaconClient.cpp
@@ -203,6 +203,16 @@ FString AMultiServerBeaconClient::GetRemotePeerId() const
 	return RemotePeerId;
 }
 
+FString AMultiServerBeaconClient::GetLocalPeerId() const
+{
+	if (ensure(OwningNode))
+	{
+		return OwningNode->GetLocalPeerId();
+	}
+
+	return TEXT("OwningNode was nullptr");
+}
+
 bool AMultiServerBeaconClient::IsAuthorityBeacon() const
 {
 	return !GetNetDriver() || GetNetDriver()->ServerConnection == nullptr;
diff --git a/Engine/Plugins/Runtime/MultiServerReplication/Source/MultiServerReplication/Public/MultiServerBeaconClient.h b/Engine/Plugins/Runtime/MultiServerReplication/Source/MultiServerReplication/Public/MultiServerBeaconClient.h
--- a/Engine/Plugins/Runtime/MultiServerReplication/Source/MultiServerReplication/Public/MultiServerBeaconClient.h
+++ b/Engine/Plugins/Runtime/MultiServerReplication/Source/MultiServerReplication/Public/MultiServerBeaconClient.h
@@ -52,8 +52,12 @@ public:
 	UFUNCTION(Reliable, Server, WithValidation, SealedEvent)
 	void ServerUpdateMultipleLevelsVisibility(const TArray<FUpdateLevelVisibilityLevelInfo>& LevelVisibilities);
 
+	/** The Id on the other side of the connection */
 	FString GetRemotePeerId() const;
 
+	/** The Id on the local side of the connection */
+	FString GetLocalPeerId() const;
+
 	/**
 	 * Actor Role on client beacons stays as authority, so this function can be used to determine whether this instance is acting as an authority.
 	 * Currently, can be used to know whether a Client or Server RPC should be called from this instance.
