Ensure that proxy connection player ids are propagated to connections on the game server.

diff --git a/Engine/Plugins/Runtime/MultiServerReplication/Source/MultiServerReplication/Private/MultiServerProxy.cpp b/Engine/Plugins/Runtime/MultiServerReplication/Source/MultiServerReplication/Private/MultiServerProxy.cpp
--- a/Engine/Plugins/Runtime/MultiServerReplication/Source/MultiServerReplication/Private/MultiServerProxy.cpp
+++ b/Engine/Plugins/Runtime/MultiServerReplication/Source/MultiServerReplication/Private/MultiServerProxy.cpp
@@ -38,6 +38,16 @@ void UGameServerNotify::NotifyControlMessage(UNetConnection* GameServerConnectio
 	}
 }
 
+void UGameServerNotify::SetFirstPlayer(TObjectPtr<ULocalPlayer> Player)
+{
+	FirstPlayer = Player;
+}
+
+ULocalPlayer* UGameServerNotify::GetFirstGamePlayer()
+{
+	return FirstPlayer;
+}
+
 void UProxyListenerNotify::SetProxyNetDriver(TObjectPtr<class UProxyNetDriver> InProxyNetDriver)
 {
 	ProxyNetDriver = InProxyNetDriver;
@@ -62,7 +72,7 @@ void UProxyListenerNotify::NotifyControlMessage(UNetConnection* ProxyConnection,
 {
 	check(ProxyNetDriver);
 
-	auto ProcessConnectionJoin = [this, ProxyConnection](int32 GameServerConnectionStateIndex, FGameServerConnectionState* GameServerConnectionState)
+	auto ProcessConnectionJoin = [this, ProxyConnection](int32 GameServerConnectionStateIndex, FGameServerConnectionState* GameServerConnectionState, bool bIsPrimaryGameServer)
 	{
 		if (GameServerConnectionState->NetDriver == nullptr)
 		{
@@ -85,15 +95,15 @@ void UProxyListenerNotify::NotifyControlMessage(UNetConnection* ProxyConnection,
 		ProxyNetDriver->GetWorld()->AddController(NewController);
 					
 		// Add a player to use on the game server.
-		static int32 InternalUserId = 10;
-		FPlatformUserId UserId = FPlatformUserId::CreateFromInternalId(InternalUserId++);
-		ULocalPlayer* NewPlayer = NewObject<ULocalPlayer>(GEngine, GEngine->LocalPlayerClass);
+		const FPlatformUserId GameServerClientId = FPlatformUserId::CreateFromInternalId(ProxyNetDriver->GetNextGameServerClientId());
+		ULocalPlayer* NewPlayer = NewObject<ULocalPlayer>(GEngine, ULocalPlayer::StaticClass());
 		NewController->SetPlayer(NewPlayer);
-		ProxyNetDriver->GetWorld()->GetGameInstance()->AddLocalPlayer(NewPlayer, UserId);
+		ProxyNetDriver->GetWorld()->GetGameInstance()->AddLocalPlayer(NewPlayer, GameServerClientId);
 
-		// The player's primary game server is currently chosen as the first game server in the list
-		// of registered game servers.
-		bool bIsPrimaryGameServer = (ProxyNetDriver->ProxyConnectionToPrimaryPlayer.Find(ProxyConnection) == nullptr);
+		// The new player will use the same unique identifier as the incoming proxy connection so that it will be propogated up to 
+		// the game servers through UNetConnection::PlayerId. This way each game server's incoming connection will have a PlayerId 
+		// that corresponds to a client connected to the proxy.
+		NewPlayer->SetCachedUniqueNetId(ProxyConnection->PlayerId);
 
 		UNetConnection* GameServerConnection = GameServerConnectionState->NetDriver->ServerConnection;
 		const bool bIsFirstGameServerConnection = (GameServerConnection == nullptr);
@@ -114,6 +124,7 @@ void UProxyListenerNotify::NotifyControlMessage(UNetConnection* ProxyConnection,
 			GameServerConnectionState->GameServerNotify = NewObject<UGameServerNotify>();
 			GameServerConnectionState->GameServerNotify->Initialize(URL);
 			GameServerConnectionState->GameServerNotify->InitNetDriver(GameServerConnectionState->NetDriver);
+			GameServerConnectionState->GameServerNotify->SetFirstPlayer(NewPlayer);
 			GameServerConnection = GameServerConnectionState->NetDriver->ServerConnection;
 
 			ProxyNetDriver->PrepareGameServerConnection(GameServerConnection);
@@ -182,6 +193,8 @@ void UProxyListenerNotify::NotifyControlMessage(UNetConnection* ProxyConnection,
 
 		if (bIsPrimaryGameServer)
 		{
+			ensure(ProxyNetDriver->ProxyConnectionToPrimaryPlayer.Find(ProxyConnection) == nullptr);
+
 			ProxyConnection->PlayerController = NewController;
 			ProxyConnection->OwningActor = NewController;
 
@@ -199,7 +212,14 @@ void UProxyListenerNotify::NotifyControlMessage(UNetConnection* ProxyConnection,
 
 	if (MessageType == NMT_Join)
 	{
-		ProxyNetDriver->ForEachRegisteredGameServer(ProcessConnectionJoin);
+		// For now, the primary game server is always the first registered game server.
+		const int32_t PrimaryGameServerIndex = 0;
+
+		for (int32 Index = 0; Index < ProxyNetDriver->GetGameServerConnectionCount(); Index++)
+		{
+			const bool bIsPrimaryGameServer = Index == PrimaryGameServerIndex;
+			ProcessConnectionJoin(PrimaryGameServerIndex, ProxyNetDriver->GetGameServerConnection(Index), bIsPrimaryGameServer);
+		}
 	}
 	else
 	{
@@ -278,12 +298,19 @@ bool UProxyNetDriver::IsConnectedToAllGameServers() const
 	return true;
 }
 
-void UProxyNetDriver::ForEachRegisteredGameServer(TFunction<void(int32, FGameServerConnectionState*)> Func)
+int32 UProxyNetDriver::GetGameServerConnectionCount() const
+{
+	return GameServerConnections.Num();
+}
+
+FGameServerConnectionState* UProxyNetDriver::GetGameServerConnection(int32 Index)
 {
-	for (int32 GameServerConnectionIndex = 0; GameServerConnectionIndex < GameServerConnections.Num(); GameServerConnectionIndex++)
+	if (ensure(Index < GameServerConnections.Num()))
 	{
-		Func(GameServerConnectionIndex, &(GameServerConnections.GetData()[GameServerConnectionIndex]));
+		return &GameServerConnections[Index];
 	}
+
+	return nullptr;
 }
 
 bool UProxyNetDriver::InitBase(bool bInitAsClient, FNetworkNotify* InNotify, const FURL& URL, bool bReuseAddressAndPort, FString& Error)
@@ -401,6 +428,11 @@ void UProxyNetDriver::AddNetworkActor(AActor* Actor)
 	Super::AddNetworkActor(Actor);
 }
 
+int32 UProxyNetDriver::GetNextGameServerClientId()
+{
+	return NextGameServerClientId++;
+}
+
 ULocalPlayer* UProxyNetDriver::GetPlayer(UNetConnection* GameServerConnection, APlayerController* PlayerController)
 {
 	const int32 PlayerIndex = PlayerController->NetPlayerIndex;
diff --git a/Engine/Plugins/Runtime/MultiServerReplication/Source/MultiServerReplication/Public/MultiServerProxy.h b/Engine/Plugins/Runtime/MultiServerReplication/Source/MultiServerReplication/Public/MultiServerProxy.h
--- a/Engine/Plugins/Runtime/MultiServerReplication/Source/MultiServerReplication/Public/MultiServerProxy.h
+++ b/Engine/Plugins/Runtime/MultiServerReplication/Source/MultiServerReplication/Public/MultiServerProxy.h
@@ -83,6 +83,14 @@ public:
 	virtual EAcceptConnection::Type NotifyAcceptingConnection() override;
 	virtual bool 					NotifyAcceptingChannel(class UChannel* Channel) override;
 	virtual void 					NotifyControlMessage(UNetConnection* GameServerConnection, uint8 MessageType, class FInBunch& Bunch) override;
+	virtual class ULocalPlayer* 	GetFirstGamePlayer();
+
+	/** Set the value to be returned by GetFirstGamePlayer(). */
+	void SetFirstPlayer(TObjectPtr<class ULocalPlayer> Player);
+
+private:
+	UPROPERTY()
+	TObjectPtr<class ULocalPlayer> FirstPlayer;
 };
 
 /** 
@@ -178,7 +186,8 @@ public:
 	bool IsConnectedToAllGameServers() const;
 
 	/** Enumerate through all outgoing connections to game servers. */
-	void ForEachRegisteredGameServer(TFunction<void(int32, FGameServerConnectionState*)> Func);
+	int32 GetGameServerConnectionCount() const;
+	FGameServerConnectionState* GetGameServerConnection(int32 Index);
 
 	virtual bool InitBase(bool bInitAsClient, FNetworkNotify* InNotify, const FURL& URL, bool bReuseAddressAndPort, FString& Error) override;
 	virtual bool InitConnect(class FNetworkNotify* InNotify, const FURL& ConnectURL, FString& Error) override;
@@ -187,13 +196,18 @@ public:
 	virtual void NotifyActorChannelOpen(UActorChannel* Channel, AActor* Actor) override;
 	virtual void AddNetworkActor(AActor* Actor) override;
 
+	/** Get the next identifier for incoming client connections to the proxy. */
+	int32 GetNextProxyClientId();
+	
+	/** Get the next identifier for outgoing connections to game servers. */
+	int32 GetNextGameServerClientId();
+
 private:
 
 	friend class UProxyListenerNotify;
 
 	// Classes in EngineTest that test the behavior of UMultiServerProxy.
 	friend class FNetTestProxyRPCRouting;
-	friend class FNetTestProxyClientValidation;
 	friend class FNetTestProxyConnectionRouting;
 	
 	friend class UProxyBackendNetDriver;
@@ -216,7 +230,7 @@ private:
 	UPROPERTY()
 	TMap<FGameServerAndPlayerIndex, TObjectPtr<UNetConnection>> GameToProxyConnection;
 
-	/** Map a proxy connection to a proxy player. */
+	/** Map a proxy connection to a primary proxy player. */
 	UPROPERTY()
 	TMap<TObjectPtr<UNetConnection>, TObjectPtr<class ULocalPlayer>> ProxyConnectionToPrimaryPlayer;
 
@@ -235,4 +249,6 @@ private:
 	/** Proxy listener handshake logic. */
 	UPROPERTY()
 	TObjectPtr<UProxyListenerNotify> ProxyListenerNotify;
+
+	int32 NextGameServerClientId = 0;
 };
diff --git a/Engine/Source/Runtime/Engine/Classes/Engine/PendingNetGame.h b/Engine/Source/Runtime/Engine/Classes/Engine/PendingNetGame.h
--- a/Engine/Source/Runtime/Engine/Classes/Engine/PendingNetGame.h
+++ b/Engine/Source/Runtime/Engine/Classes/Engine/PendingNetGame.h
@@ -151,4 +151,8 @@ public:
 
 	/** Called by the engine after loadmapCompleted and the GameInstance has finished delaying */
 	virtual void TravelCompleted(UEngine* Engine, FWorldContext& Context);
+
+protected:
+
+	virtual class ULocalPlayer* GetFirstGamePlayer();
 };
diff --git a/Engine/Source/Runtime/Engine/Private/NetDriver.cpp b/Engine/Source/Runtime/Engine/Private/NetDriver.cpp
--- a/Engine/Source/Runtime/Engine/Private/NetDriver.cpp
+++ b/Engine/Source/Runtime/Engine/Private/NetDriver.cpp
@@ -8699,11 +8699,20 @@ FConsoleCommandWithArgsDelegate::CreateLambda([](const TArray< FString >& Args)
 		{
 			UE_LOG(LogNet, Display, TEXT("Printing Server Connection for: %s (Definition=%s)"), *NetDriver->NetDriverName.ToString(), *NetDriver->GetNetDriverDefinition().ToString());
 
-			UE_LOG(LogNet, Display, TEXT("\tServerConnection: ConnectionId=%u ViewTarget=%s FullDescription=%s"), 
+			UE_LOG(LogNet, Display, TEXT("\tServerConnection: ConnectionId=%u ViewTarget=%s FullDescription=%s Child=false"), 
 				NetDriver->ServerConnection->GetConnectionHandle().GetParentConnectionId(), 
 				*GetNameSafe(NetDriver->ServerConnection->ViewTarget),
 				*NetDriver->ServerConnection->Describe()
 			);
+
+			for (UNetConnection* ChildNetConnection : NetDriver->ServerConnection->Children)
+			{
+				UE_LOG(LogNet, Display, TEXT("\tServerConnection: ConnectionId=%u ViewTarget=%s FullDescription=%s Child=true"), 
+					   ChildNetConnection->GetConnectionHandle().GetParentConnectionId(), 
+					   *GetNameSafe(ChildNetConnection->ViewTarget),
+					   *ChildNetConnection->Describe()
+				);
+			}
 		}
 		else
 		{
@@ -8711,12 +8720,22 @@ FConsoleCommandWithArgsDelegate::CreateLambda([](const TArray< FString >& Args)
 
 			for (UNetConnection* NetConnection : NetDriver->ClientConnections)
 			{
-				UE_LOG(LogNet, Display, TEXT("\tClientConnection: ConnectionId=%u ViewTarget=%s NetId=%s FullDescription=%s"), 
+				UE_LOG(LogNet, Display, TEXT("\tClientConnection: ConnectionId=%u ViewTarget=%s NetId=%s FullDescription=%s Child=false"), 
 					NetConnection->GetConnectionHandle().GetParentConnectionId(), 
 					*GetNameSafe(NetConnection->ViewTarget),
 					*NetConnection->PlayerId.ToDebugString(), 
 					*NetConnection->Describe()
 				);
+
+				for (UNetConnection* ChildNetConnection : NetConnection->Children)
+				{
+					UE_LOG(LogNet, Display, TEXT("\tClientConnection: ConnectionId=%u ViewTarget=%s NetId=%s FullDescription=%s Child=true"), 
+						   ChildNetConnection->GetConnectionHandle().GetParentConnectionId(), 
+						   *GetNameSafe(ChildNetConnection->ViewTarget),
+						   *ChildNetConnection->PlayerId.ToDebugString(), 
+						   *ChildNetConnection->Describe()
+					);
+				}
 			}
 		}
 	};
diff --git a/Engine/Source/Runtime/Engine/Private/PendingNetGame.cpp b/Engine/Source/Runtime/Engine/Private/PendingNetGame.cpp
--- a/Engine/Source/Runtime/Engine/Private/PendingNetGame.cpp
+++ b/Engine/Source/Runtime/Engine/Private/PendingNetGame.cpp
@@ -65,7 +65,7 @@ void UPendingNetGame::InitNetDriver(UNetDriver* InNetDriver)
 		{
 			FNetDelegates::OnPendingNetGameConnectionCreated.Broadcast(this);
 
-			ULocalPlayer* LocalPlayer = GEngine->GetFirstGamePlayer(this);
+			ULocalPlayer* LocalPlayer = GetFirstGamePlayer();
 			if (LocalPlayer)
 			{
 				LocalPlayer->PreBeginHandshake(ULocalPlayer::FOnPreBeginHandshakeCompleteDelegate::CreateWeakLambda(this,
@@ -312,7 +312,7 @@ void UPendingNetGame::NotifyControlMessage(UNetConnection* Connection, uint8 Mes
 					}
 				}
 
-				ULocalPlayer* LocalPlayer = GEngine->GetFirstGamePlayer(this);
+				ULocalPlayer* LocalPlayer = GetFirstGamePlayer();
 				if (LocalPlayer)
 				{
 					// Send the player nickname if available
@@ -566,3 +566,7 @@ void UPendingNetGame::SendJoin()
 	NetDriver->ServerConnection->FlushNet(true);
 }
 
+ULocalPlayer* UPendingNetGame::GetFirstGamePlayer()
+{
+	return GEngine->GetFirstGamePlayer(this);
+}
\ No newline at end of file
