Update MultiServerReplication plugin metadata.

diff --git a/Engine/Plugins/Runtime/MultiServerReplication/MultiServerReplication.uplugin b/Engine/Plugins/Runtime/MultiServerReplication/MultiServerReplication.uplugin
--- a/Engine/Plugins/Runtime/MultiServerReplication/MultiServerReplication.uplugin
+++ b/Engine/Plugins/Runtime/MultiServerReplication/MultiServerReplication.uplugin
@@ -1,10 +1,10 @@
 {
 	"FileVersion": 3,
-	"FriendlyName": "MultiServerReplication Plugin",
+	"FriendlyName": "Multi-server Replication",
 	"Version": 1,
 	"VersionName": "1.0",
-	"Description": "",
-	"Category": "Epic Internal",
+	"Description": "Code to help facilitate connecting multiple UE server processes to each other.",
+	"Category": "Networking",
 	"CreatedBy": "Epic Games, Inc.",
 	"CreatedByURL": "http://epicgames.com",
 	"EnabledByDefault": false,
