Allow unlimited bunch size for the MultiServerReplication plugin.

diff --git a/Engine/Plugins/Runtime/MultiServerReplication/Source/MultiServerReplication/Private/MultiServerBeaconClient.cpp b/Engine/Plugins/Runtime/MultiServerReplication/Source/MultiServerReplication/Private/MultiServerBeaconClient.cpp
--- a/Engine/Plugins/Runtime/MultiServerReplication/Source/MultiServerReplication/Private/MultiServerBeaconClient.cpp
+++ b/Engine/Plugins/Runtime/MultiServerReplication/Source/MultiServerReplication/Private/MultiServerBeaconClient.cpp
@@ -53,6 +53,11 @@ void AMultiServerBeaconClient::OnConnected()
 		ServerSetRemotePeerId(OwningNode->GetLocalPeerId());
 	}
 
+	if (UNetConnection* Connection = GetNetConnection())
+	{
+		Connection->SetUnlimitedBunchSizeAllowed(true);
+	}
+
 	UWorld* World = GetWorld();
 	if (World)
 	{
diff --git a/Engine/Plugins/Runtime/MultiServerReplication/Source/MultiServerReplication/Private/MultiServerBeaconHostObject.cpp b/Engine/Plugins/Runtime/MultiServerReplication/Source/MultiServerReplication/Private/MultiServerBeaconHostObject.cpp
--- a/Engine/Plugins/Runtime/MultiServerReplication/Source/MultiServerReplication/Private/MultiServerBeaconHostObject.cpp
+++ b/Engine/Plugins/Runtime/MultiServerReplication/Source/MultiServerReplication/Private/MultiServerBeaconHostObject.cpp
@@ -3,6 +3,7 @@
 #include "MultiServerBeaconHostObject.h"
 #include "MultiServerBeaconClient.h"
 #include "MultiServerNode.h"
+#include "Engine/NetConnection.h"
 
 #include UE_INLINE_GENERATED_CPP_BY_NAME(MultiServerBeaconHostObject)
 
@@ -25,6 +26,9 @@ void AMultiServerBeaconHostObject::OnClientConnected(AOnlineBeaconClient* NewCli
 	Super::OnClientConnected(NewClientActor, ClientConnection);
 
 	UE_LOG(LogMultiServerBeacon, Verbose, TEXT("AMultiServerBeaconHostObject::OnClientConnected: num client actors connected is %d"), ClientActors.Num());
+
+	ClientConnection->SetUnlimitedBunchSizeAllowed(true);
+
 	if (AMultiServerBeaconClient* MultiServerClient = Cast<AMultiServerBeaconClient>(NewClientActor))
 	{
 		MultiServerClient->SetOwningNode(OwningNode);
