Refactor the multi-server proxy to use it's own player controllers for incoming connections and support for player migration on backend game servers.

diff --git a/Engine/Plugins/Runtime/MultiServerReplication/Source/MultiServerReplication/Private/MultiServerProxy.cpp b/Engine/Plugins/Runtime/MultiServerReplication/Source/MultiServerReplication/Private/MultiServerProxy.cpp
--- a/Engine/Plugins/Runtime/MultiServerReplication/Source/MultiServerReplication/Private/MultiServerProxy.cpp
+++ b/Engine/Plugins/Runtime/MultiServerReplication/Source/MultiServerReplication/Private/MultiServerProxy.cpp
@@ -4,7 +4,9 @@
 #include "UnrealEngine.h"
 #include "Engine/LocalPlayer.h"
 #include "Engine/GameInstance.h"
+#include "Engine/NetworkObjectList.h"
 #include "Engine/World.h"
+#include "GameFramework/GameModeBase.h"
 #include "Misc/CommandLine.h"
 #include "Net/DataChannel.h"
 
@@ -74,6 +76,8 @@ void UProxyListenerNotify::NotifyControlMessage(UNetConnection* ProxyConnection,
 
 	auto ProcessConnectionJoin = [this, ProxyConnection](int32 GameServerConnectionStateIndex, FGameServerConnectionState* GameServerConnectionState, bool bIsPrimaryGameServer)
 	{
+		// If this is the first connection to the game server, instantiate the backend network driver that will manage all
+		// connections from proxy connections to that server.
 		if (GameServerConnectionState->NetDriver == nullptr)
 		{
 			static int32 GameServerDriverIndex = 0;
@@ -86,18 +90,12 @@ void UProxyListenerNotify::NotifyControlMessage(UNetConnection* ProxyConnection,
 
 			UE_LOG(LogNetProxy, Log, TEXT("Created a game server NetDriver (name=%s, port=%d)"), *GameServerConnectionState->NetDriver->GetName(), GameServerConnectionState->Port);
 		}
+		
+		UE_LOG(LogNetProxy, Log, TEXT("Client starting connection handshake on %s:%s"), *ProxyConnection->Driver->GetName(), *ProxyConnection->GetName());
 
-		// Spawn a temporary, placeholder controller and add it to the world.
-		FActorSpawnParameters SpawnInfo;
-		SpawnInfo.ObjectFlags |= RF_Transient;
-		SpawnInfo.bDeferConstruction = true;
-		APlayerController* NewController = ProxyNetDriver->GetWorld()->SpawnActor<APlayerController>(APlayerController::StaticClass(), SpawnInfo);
-		ProxyNetDriver->GetWorld()->AddController(NewController);
-					
 		// Add a player to use on the game server.
 		const FPlatformUserId GameServerClientId = FPlatformUserId::CreateFromInternalId(ProxyNetDriver->GetNextGameServerClientId());
 		ULocalPlayer* NewPlayer = NewObject<ULocalPlayer>(GEngine, ULocalPlayer::StaticClass());
-		NewController->SetPlayer(NewPlayer);
 		ProxyNetDriver->GetWorld()->GetGameInstance()->AddLocalPlayer(NewPlayer, GameServerClientId);
 
 		// The new player will use the same unique identifier as the incoming proxy connection so that it will be propogated up to 
@@ -127,26 +125,14 @@ void UProxyListenerNotify::NotifyControlMessage(UNetConnection* ProxyConnection,
 			GameServerConnectionState->GameServerNotify->SetFirstPlayer(NewPlayer);
 			GameServerConnection = GameServerConnectionState->NetDriver->ServerConnection;
 
-			ProxyNetDriver->PrepareGameServerConnection(GameServerConnection);
-
 			// UNetDriver::Notify will be reset in UPendingNetGame above so it's important that we override
 			// it here again to point to the proxy.
 			GameServerConnectionState->NetDriver->Notify = GameServerConnectionState->GameServerNotify;
 
-			// Use this placeholder controller on the incoming connection to the proxy and the outgoing
-			// connection to the game server.
-			GameServerConnection->PlayerController = NewController;
-			GameServerConnection->OwningActor = NewController;
-					
-			// Register a delegate callback when a new connection is successfully established with the game server
-			// to bind the connection to the game server to the proxy connection using a common player controller.
-			GameServerConnection->GetOnConnectionOpened().AddUObject(ProxyNetDriver, &UProxyNetDriver::ConnectionHandshakeComplete);
-
-			UE_LOG(LogNetProxy, Log, TEXT("Connecting to game server: %s:%s -> %s (controller=%s player=%s"), 
+			UE_LOG(LogNetProxy, Log, TEXT("Connecting to game server: %s:%s -> %s (player=%s)"), 
 				   *ProxyConnection->GetName(),
 				   *ProxyConnection->GetDriver()->GetName(), 
 				   *GameServerConnection->GetDriver()->GetName(),
-				   *NewController->GetName(),
 				   *NewPlayer->GetName());
 		}
 		else
@@ -157,28 +143,26 @@ void UProxyListenerNotify::NotifyControlMessage(UNetConnection* ProxyConnection,
 
 			// ULocalPlayer::SendSplitJoin() can only be called when the parent connection has been successfully opened. If 
 			// the parent connection is not open yet, this player will be queued up and processed when the parent connection
-			// has been opened (see UMultiProxyServer::ConnectionHandshakeComplete()).
+			// has been opened (see UMultiProxyServer::GameServerConnectionHandshakeComplete()).
 			if (GameServerConnectionState->NetDriver->ServerConnection->GetConnectionState() == USOCK_Open)
 			{
 				GameServerConnectionState->PendingSplitJoinRequests.Add(Request);
 				ProxyNetDriver->FlushSplitJoinRequests(GameServerConnectionState);
 
-				UE_LOG(LogNetProxy, Log, TEXT("Connecting (multiplexed) to game server: %s:%s -> %s (controller=%s player=%s"), 
+				UE_LOG(LogNetProxy, Log, TEXT("Connecting (multiplexed) to game server: %s:%s -> %s (player=%s"), 
 					   *ProxyConnection->GetName(),
 					   *ProxyConnection->GetDriver()->GetName(), 
 					   *GameServerConnection->GetDriver()->GetName(),
-					   *NewController->GetName(),
 					   *NewPlayer->GetName());
 			}
 			else
 			{
 				GameServerConnectionState->PendingSplitJoinRequests.Add(Request);
 	
-				UE_LOG(LogNetProxy, Log, TEXT("Queuing up connection (multiplexed) to game server: %s:%s -> %s (controller=%s player=%s"), 
+				UE_LOG(LogNetProxy, Log, TEXT("Queuing up connection (multiplexed) to game server: %s:%s -> %s (player=%s)"), 
 					   *ProxyConnection->GetName(),
-					   *ProxyConnection->GetDriver()->GetName(), 
+					   *ProxyConnection->GetDriver()->GetName(),
 					   *GameServerConnection->GetDriver()->GetName(),
-					   *NewController->GetName(),
 					   *NewPlayer->GetName());
 			}
 		}
@@ -186,6 +170,9 @@ void UProxyListenerNotify::NotifyControlMessage(UNetConnection* ProxyConnection,
 		const int32 NewPlayerIndex = GameServerConnectionState->Players.Num();
 		GameServerConnectionState->Players.Add(NewPlayer);
 		
+		// Associate this proxy connection with the parent game server connection when beginning the handshake because
+		// the child connection hasn't been created yet. Once the handshake is complete, it's expected that this entry will
+		// be updated with the new child connection.
 		FGameServerAndPlayerIndex Key; 
 		Key.GameServerConnection = GameServerConnection;
 		Key.PlayerIndex = NewPlayerIndex;
@@ -193,20 +180,7 @@ void UProxyListenerNotify::NotifyControlMessage(UNetConnection* ProxyConnection,
 
 		if (bIsPrimaryGameServer)
 		{
-			ensure(ProxyNetDriver->ProxyConnectionToPrimaryPlayer.Find(ProxyConnection) == nullptr);
-
-			ProxyConnection->PlayerController = NewController;
-			ProxyConnection->OwningActor = NewController;
-
-			UE_LOG(LogNetProxy, Log, TEXT("Configuring primary game server %s for player (player=%s:%d proxy_connection=%s:%s)"),
-				   *GameServerConnection->GetDriver()->GetName(),
-				   *NewPlayer->GetName(),
-				   NewPlayerIndex,
-				   *ProxyConnection->GetDriver()->GetName(),
-				   *ProxyConnection->GetName());
-
-			ProxyNetDriver->ProxyPlayerToPrimaryGameServer.Add(NewPlayer, GameServerConnectionStateIndex);
-			ProxyNetDriver->ProxyConnectionToPrimaryPlayer.Add(ProxyConnection, NewPlayer);
+			ProxyNetDriver->RegisterPrimaryGameServer(ProxyConnection, NewPlayer, GameServerConnectionStateIndex);
 		}
 	};
 
@@ -217,7 +191,7 @@ void UProxyListenerNotify::NotifyControlMessage(UNetConnection* ProxyConnection,
 
 		for (int32 Index = 0; Index < ProxyNetDriver->GetGameServerConnectionCount(); Index++)
 		{
-			const bool bIsPrimaryGameServer = Index == PrimaryGameServerIndex;
+			const bool bIsPrimaryGameServer = (Index == PrimaryGameServerIndex);
 			ProcessConnectionJoin(PrimaryGameServerIndex, ProxyNetDriver->GetGameServerConnection(Index), bIsPrimaryGameServer);
 		}
 	}
@@ -228,17 +202,46 @@ void UProxyListenerNotify::NotifyControlMessage(UNetConnection* ProxyConnection,
 	}
 }
 
+void UProxyBackendNetConnection::HandleClientPlayer(APlayerController* NewPlayerController, UNetConnection* GameServerConnection)
+{
+	ensure(GameServerConnection == this);
+
+	UProxyBackendNetDriver* BackendNetDriver = Cast<UProxyBackendNetDriver>(Driver);
+	if (ensure(BackendNetDriver))
+	{
+		BackendNetDriver->GetProxyNetDriver()->GameServerConnectionHandshakeComplete(this, NewPlayerController, GameServerConnection);
+	}
+}
+
+void UProxyBackendChildNetConnection::HandleClientPlayer(APlayerController* NewPlayerController, UNetConnection* GameServerConnection)
+{
+	ensure(GameServerConnection != this);
+
+	UProxyBackendNetDriver* BackendNetDriver = Cast<UProxyBackendNetDriver>(Driver);
+	if (ensure(BackendNetDriver))
+	{
+		BackendNetDriver->GetProxyNetDriver()->GameServerConnectionHandshakeComplete(this, NewPlayerController, GameServerConnection);
+	}
+}
+
 void UProxyBackendNetDriver::SetProxyNetDriver(TObjectPtr<UProxyNetDriver> InProxyNetDriver)
 {
 	ProxyNetDriver = InProxyNetDriver;
 }
 
+TObjectPtr<UProxyNetDriver> UProxyBackendNetDriver::GetProxyNetDriver()
+{
+	return ProxyNetDriver;
+}
+
 bool UProxyBackendNetDriver::InitBase(bool bInitAsClient, FNetworkNotify* InNotify, const FURL& URL, bool bReuseAddressAndPort, FString& Error)
 {
 	bool bSuccess = Super::InitBase(bInitAsClient, InNotify, URL, bReuseAddressAndPort, Error);
 
 	if (bSuccess)
 	{
+		NetConnectionClass = UProxyBackendNetConnection::StaticClass();
+
 		// Don't allow any RPCs received game servers to be executed on the proxy.
 		EnableExecuteRPCFunctions(false);
 	}
@@ -258,19 +261,101 @@ void UProxyBackendNetDriver::ForwardRemoteFunction(UObject* RootObject, UObject*
 	AActor* OwningActor = Cast<AActor>(RootObject);
 	if (!ensure(OwningActor != nullptr))
 	{
+		UE_LOG(LogNetProxy, Verbose, TEXT("RPC: Ignoring %s on %s from game server %s because it has no owner."), *Function->GetName(), *RootObject->GetName(), *GetName());
 		return;
 	}
 
+	// If the owner is a PlayerController it is one that represents a connection to a game server and must be mapped
+	// to a PlayerController for an incoming proxy connection. If no such mapping exists it means that the RPC comes 
+	// from a game server that isn't the primary game server.
+	if (APlayerController* GameServerController = Cast<APlayerController>(OwningActor))
+	{
+		APlayerController* ProxyController = ProxyNetDriver->GetProxyControllerFromPrimaryGameServerController(GameServerController);
+		if (ProxyController == nullptr)
+		{
+			UE_LOG(LogNetProxy, Verbose, TEXT("RPC: Ignoring %s on %s from game server %s because it's owning game server controller %s cannot map to a proxy controller."), *Function->GetName(), *RootObject->GetName(), *GetName(), *GameServerController->GetName());
+			return;
+		}
+		else
+		{
+			UE_LOG(LogNetProxy, VeryVerbose, TEXT("RPC: Remapping game server controller %s to proxy controller %s when forwarding RPC %s on %s."), *GameServerController->GetName(), *ProxyController->GetName(), *Function->GetName(), *RootObject->GetName());
+		}
+
+		OwningActor = ProxyController;
+	}
+
 	ULocalPlayer* OwningPlayer = Cast<ULocalPlayer>(OwningActor->GetNetOwningPlayerAnyRole());
 	if (OwningPlayer == nullptr)
 	{
+		UE_LOG(LogNetProxy, Verbose, TEXT("RPC: Ignoring %s on %s from game server %s because it's owning actor %s doesn't have an owning player."), *Function->GetName(), *RootObject->GetName(), *GetName(), *OwningActor->GetName());
 		return;
 	}
 
-	int32* GameServerConnectionIndex = ProxyNetDriver->ProxyPlayerToPrimaryGameServer.Find(OwningPlayer);
-	if (GameServerConnectionIndex && *GameServerConnectionIndex < ProxyNetDriver->GameServerConnections.Num())
+	UE_LOG(LogNetProxy, VeryVerbose, TEXT("RPC: Pre-forwarding %s on %s from game server %s to proxy connection."), *Function->GetName(), *OwningActor->GetName(), *GetName());
+
+	// The OwningActor actor will be associated with a connection (the return value of AActor::GetNetConnection()) to the backend game server, but
+	// when UProxyNetDriver::InternalProcessRemoteFunction() is called the owning player will be used to lookup the actual proxy connection to
+	// forward the RPC.
+	ProxyNetDriver->ProcessRemoteFunction(OwningActor, Function, Parms, static_cast<FOutParmRec*>(nullptr), static_cast<FFrame*>(nullptr), SubObject);
+}
+
+bool UProxyBackendNetDriver::ShouldSkipRepNotifies() const
+{
+	return true;
+}
+
+UChildConnection* UProxyBackendNetDriver::CreateChild(UNetConnection* Parent)
+{
+	UChildConnection* Child = NewObject<UProxyBackendChildNetConnection>();
+	Child->InitChildConnection(this, Parent);
+	Parent->Children.Add(Child);
+	return Child;
+}
+
+
+void UProxyBackendNetDriver::InternalProcessRemoteFunction(
+	class AActor* Actor,
+	class UObject* SubObject,
+	class UNetConnection* Connection,
+	class UFunction* Function,
+	void* Parms,
+	FOutParmRec* OutParms,
+	FFrame* Stack,
+	bool bIsServer)
+{
+	AActor* DestinationActor = Actor;
+
+	// The majority of the actors on the proxy are replicated from the game servers so their owners will be game server player
+	// controllers. The exception are proxy player controllers which are spawned by the proxy and will have an owner that is
+	// the proxy player controller. In this case we must forward the RPC to the primary game server player controller.
+	if (ProxyNetDriver->IsProxySpawned(Actor))
+	{
+		if (APlayerController* ProxyPlayerController = Cast<APlayerController>(Actor))
+		{
+			if (ULocalPlayer* Player = Cast<ULocalPlayer>(ProxyPlayerController->Player))
+			{
+				if (TObjectPtr<APlayerController>* GameServerPlayerControllerPtr = ProxyNetDriver->ProxyPlayerToPrimaryGameServerPlayerController.Find(Player))
+				{
+					UNetConnection* GameServerConnection = GameServerPlayerControllerPtr->Get()->NetConnection;
+					ensure(GameServerConnection->PlayerController->Player == Player);
+					DestinationActor = GameServerConnection->PlayerController;
+					SubObject = nullptr;
+
+					UE_LOG(LogNetProxy, VeryVerbose, TEXT("RPC: Remapping proxy player controller %s to game server player controller %s when forwarding RPC %s on %s"), 
+						   *ProxyPlayerController->GetName(), *GameServerConnection->PlayerController->GetName(), *Function->GetName(), *DestinationActor->GetName());
+				}
+			}
+		}
+	}
+
+	if (!ProxyNetDriver->IsProxySpawned(DestinationActor))
+	{
+		UE_LOG(LogNetProxy, Verbose, TEXT("RPC: Forwarding %s on %s to game server connection %s"), *Function->GetName(), *DestinationActor->GetName(), *Connection->GetName());
+		Super::InternalProcessRemoteFunction(DestinationActor, SubObject, Connection, Function, Parms, OutParms, Stack, bIsServer);
+	}
+	else
 	{
-		ProxyNetDriver->ProcessRemoteFunction(OwningActor, Function, Parms, static_cast<FOutParmRec*>(nullptr), static_cast<FFrame*>(nullptr), SubObject);
+		UE_LOG(LogNetProxy, Verbose, TEXT("RPC: Ignoring %s on %s to game server connection %s because it's not owned by a game server."), *Function->GetName(), *DestinationActor->GetName(), *Connection->GetName());
 	}
 }
 
@@ -368,23 +453,27 @@ void UProxyNetDriver::ForwardRemoteFunction(UObject* RootObject, UObject* SubObj
 	// This function is called when the proxy receives an RPC from a game client and will forward the function on
 	// to the owning player's primary game server.
 
+	// If the owner is a PlayerController it will be for a proxy connection. There is no need to map it to the PlayerController
+	// on the game server because UNetDriver::ProcessRemoteFunction() will automatically send the RPC to the server connection
+	// and ignores the value of AActor::GetNetConnection().
 	AActor* OwningActor = Cast<AActor>(RootObject);
 	if (!ensure(OwningActor != nullptr))
 	{
+		UE_LOG(LogNetProxy, Verbose, TEXT("RPC: Ignoring %s on %s from proxy client because it doesn't have an owner."), *Function->GetName(), *RootObject->GetName());
 		return;
 	}
 
 	ULocalPlayer* OwningPlayer = Cast<ULocalPlayer>(OwningActor->GetNetOwningPlayerAnyRole());
 	if (OwningPlayer == nullptr)
 	{
+		UE_LOG(LogNetProxy, Verbose, TEXT("RPC: Ignoring %s on %s from proxy client because it owner %s doesn't have an owning player"), *Function->GetName(), *RootObject->GetName(), *OwningActor->GetName());
 		return;
 	}
 
-	int32* GameServerConnectionIndex = ProxyPlayerToPrimaryGameServer.Find(OwningPlayer);
-	if (GameServerConnectionIndex && *GameServerConnectionIndex < GameServerConnections.Num())
+	if (TObjectPtr<APlayerController>* GamePlayerControllerPtr = ProxyPlayerToPrimaryGameServerPlayerController.Find(OwningPlayer))
 	{
-		FGameServerConnectionState& GameServerConnection = GameServerConnections[*GameServerConnectionIndex];
-		GameServerConnection.NetDriver->ProcessRemoteFunction(OwningActor, Function, Parms, static_cast<FOutParmRec*>(nullptr), static_cast<FFrame*>(nullptr), SubObject);
+		UNetConnection* GameConnection = GamePlayerControllerPtr->Get()->NetConnection;
+		GameConnection->Driver->ProcessRemoteFunction(OwningActor, Function, Parms, static_cast<FOutParmRec*>(nullptr), static_cast<FFrame*>(nullptr), SubObject);
 	}
 }
 
@@ -416,7 +505,9 @@ void UProxyNetDriver::AddNetworkActor(AActor* Actor)
 	// loaded by the proxy and those replicated from the connected game servers. It's assumed that 
 	// the actors replicated from the game servers will not have a role of ROLE_Authority and will
 	// therefore be unaffacted by this code and replicate as normal.
-	if (Actor->GetIsReplicated())
+	//
+	// Actors that are explicitly spawned as part of the proxy functionality are allowed to be replicated.
+	if (Actor->GetIsReplicated() && !IsProxySpawned(Actor))
 	{
 		if (Actor->GetLocalRole() == ROLE_Authority)
 		{
@@ -428,97 +519,251 @@ void UProxyNetDriver::AddNetworkActor(AActor* Actor)
 	Super::AddNetworkActor(Actor);
 }
 
-int32 UProxyNetDriver::GetNextGameServerClientId()
+bool UProxyNetDriver::ShouldCallRemoteFunction(UObject* Object, UFunction* Function, const FReplicationFlags& RepFlags) const
 {
-	return NextGameServerClientId++;
+	return !RepFlags.bIgnoreRPCs;
 }
 
-ULocalPlayer* UProxyNetDriver::GetPlayer(UNetConnection* GameServerConnection, APlayerController* PlayerController)
+void UProxyNetDriver::InternalProcessRemoteFunction(
+	class AActor* Actor,
+	class UObject* SubObject,
+	class UNetConnection* Connection,
+	class UFunction* Function,
+	void* Parms,
+	FOutParmRec* OutParms,
+	FFrame* Stack,
+	bool bIsServer)
 {
-	const int32 PlayerIndex = PlayerController->NetPlayerIndex;
-	ULocalPlayer* FoundPlayer = nullptr;
+	if (ULocalPlayer* Player = Cast<ULocalPlayer>(Actor->GetNetOwningPlayerAnyRole()))
+	{
+		// RPCs from game servers will be routed to proxy connections. Only RPCs from a connections primary game server will
+		// be routed to the game client; other RPCs will be ignored.
+		if (UNetConnection* ProxyConnection = GetProxyConnectionFromPrimaryPlayer(Player))
+		{
+			UE_LOG(LogNetProxy, Verbose, TEXT("RPC: Forwarding %s on %s to proxy connection %s for player %s"), *Function->GetName(), *Actor->GetName(), *ProxyConnection->GetName(), *Player->GetName());
+			Super::InternalProcessRemoteFunction(Actor, SubObject, ProxyConnection, Function, Parms, OutParms, Stack, bIsServer);
+		}
+		else
+		{
+			UE_LOG(LogNetProxy, Verbose, TEXT("RPC: Ignoring %s on %s because player %s isn't a primary player."), *Function->GetName(), *Actor->GetName(), *Player->GetName());
+		}
+	}
+	else
+	{
+		UE_LOG(LogNetProxy, Verbose, TEXT("RPC: Ignoring %s on %s because there is no actor owning player."), *Function->GetName(), *Actor->GetName());
+	}
+}
 
-	for (FGameServerConnectionState& GameServerConnectionState : GameServerConnections)
+int32 UProxyNetDriver::ServerReplicateActors(float DeltaSeconds)
+{
+	// The owners are actors in the world are going to be game server player controlers. Since these controllers 
+	// aren't replicated ot the client, and the client has it's own proxy player controller, we need to map an
+	// acor's game server player controller owner to the corresponding proxy player controller during replication.
+	TMap<AActor*, AActor*> OriginalActorOwners;
+	for (const TSharedPtr<FNetworkObjectInfo>& ObjectInfo : GetNetworkObjectList().GetAllObjects())
 	{
-		if (GameServerConnectionState.NetDriver == GameServerConnection->GetDriver())
+		AActor* Owner = ObjectInfo->Actor->Owner;
+		if (APlayerController* GameServerController = Cast<APlayerController>(Owner))
 		{
-			if (ensure(PlayerIndex < GameServerConnectionState.Players.Num()))
+			OriginalActorOwners.Add(ObjectInfo->Actor, GameServerController);
+
+			APlayerController* ProxyController = GetProxyControllerFromPrimaryGameServerController(GameServerController);
+			AActor* Actor = const_cast<TSharedPtr<FNetworkObjectInfo>*>(&ObjectInfo)->Get()->Actor;
+			Actor->Owner = ProxyController;
+		}
+	}
+
+	int32 ActorCount = Super::ServerReplicateActors(DeltaSeconds);
+
+	// Restore actor owners to game server player controllers.
+	for (TPair<AActor*, AActor*>& Pair : OriginalActorOwners)
+	{
+		AActor* OriginalActor = Pair.Key;
+		AActor* OriginalOwner = Pair.Value;
+		if (const TSharedPtr<FNetworkObjectInfo>* ObjectInfo = GetNetworkObjectList().GetAllObjects().Find(OriginalActor))
+		{
+			AActor* Actor = const_cast<TSharedPtr<FNetworkObjectInfo>*>(ObjectInfo)->Get()->Actor;
+			Actor->Owner = OriginalOwner;
+		}
+	}
+
+	return ActorCount;
+}
+
+bool UProxyNetDriver::CanDowngradeActorRole(UNetConnection* ProxyConnection, AActor* Actor) const
+{
+	if (ULocalPlayer* Player = Cast<ULocalPlayer>(Actor->GetNetOwningPlayerAnyRole()))
+	{
+		// If this autonomous actor is owned by a player that is bound to the same proxy connection
+		// as the attached proxy player controller, don't downgrade from ROLE_AutonomousProxy to ROLE_SimulatedProxy.
+		if (Actor->GetRemoteRole() == ROLE_AutonomousProxy)
+		{
+			if (UNetConnection* PlayerProxyConnection = GetProxyConnectionFromPrimaryPlayer(Player))
 			{
-				FoundPlayer = GameServerConnectionState.Players[PlayerIndex];
+				if (PlayerProxyConnection == ProxyConnection)
+				{
+					return false;
+				}
 			}
+
+			return true;
 		}
 	}
 
-	return FoundPlayer;
+	return false;
 }
 
-void UProxyNetDriver::PrepareGameServerConnection(UNetConnection* GameServerConnection)
+int32 UProxyNetDriver::GetNextGameServerClientId()
 {
-	// Register a callback to retrieve the correct local player for this connection based on the player controller.
-	GameServerConnection->SetLocalPlayerFunction([this](
-		UNetConnection* GameServerConnection,
-		APlayerController* PlayerController)
-	{
-		return GetPlayer(GameServerConnection, PlayerController);
-	});
+	return NextGameServerClientId++;
 }
 
-void UProxyNetDriver::ConnectionHandshakeComplete(UNetConnection* ChildGameServerConnection,
-												  APlayerController* PlayerController,
-												  UNetConnection* ParentGameServerConnection)
+bool UProxyNetDriver::IsProxySpawned(AActor* Actor) const
 {
-	const int32 PlayerIndex = PlayerController->NetPlayerIndex;
-	FGameServerAndPlayerIndex Key; 
+	return ProxySpawnedActors.Contains(Actor);
+}
+
+void UProxyNetDriver::GameServerConnectionHandshakeComplete(UNetConnection* ChildGameServerConnection,
+															APlayerController* GameServerPlayerController,
+															UNetConnection* ParentGameServerConnection)
+{
+	ULocalPlayer* NewPlayer = GetPlayerOnGameServer(ParentGameServerConnection, GameServerPlayerController->NetPlayerIndex);
+	ensure(NewPlayer);
+
+	// If a new player controller is spawned on the proxy from a game server it's assumed it will replace any
+	// existing game server player controllers used by connections and players on the proxy. For this reason
+	// always detach an existing game server player controllers from the connection.
+	DetachPlayerController(ChildGameServerConnection->PlayerController);
+	AttachPlayerController(GameServerPlayerController, ChildGameServerConnection, NewPlayer);
+	
+	// The game server player controller is used as a client connection to the game server so must be ROLE_AutonomousProxy.
+	// NOTE: A proxy player controller will still be ROLE_Authority since it's spawned by the proxy.
+	GameServerPlayerController->SetRole(ROLE_AutonomousProxy);
+
+	// The player should only reference a proxy player controller if this game server connection
+	// is the primary game server for the proxy connection (decided later in this function).
+	NewPlayer->PlayerController = nullptr;
+
+	const int32 PlayerIndex = GameServerPlayerController->NetPlayerIndex;
+	FGameServerAndPlayerIndex Key;
 	Key.GameServerConnection = ParentGameServerConnection;
-	Key.PlayerIndex = PlayerIndex;
+ 	Key.PlayerIndex = PlayerIndex;
 
+	UE_LOG(LogNetProxy, VeryVerbose, TEXT("Lookup proxy connection associated with game server connection %s and player index %d (game_server_controller=%s"), 
+		   *ParentGameServerConnection->GetName(), 
+		   PlayerIndex,
+		   *GameServerPlayerController->GetName());
+
+	// Find the proxy connection that established this connection to the game server.
 	TObjectPtr<UNetConnection>* ProxyConnectionPtr = GameToProxyConnection.Find(Key);
-	if (ProxyConnectionPtr)
+	if (ensure(ProxyConnectionPtr))
 	{
 		TObjectPtr<UNetConnection> ProxyConnection = *ProxyConnectionPtr;
-		check(ProxyConnection->Driver == this);
 
-		ULocalPlayer* LocalPlayer = GetPlayer(ParentGameServerConnection, PlayerController);
-		if (ensure(LocalPlayer))
+		bool bBindProxyConnectionToPrimaryGameServer = false;
+
+		if (int32* PlayerPrimaryGameServerIdx = ProxyPlayerToPrimaryGameServer.Find(NewPlayer))
 		{
-			int32* PlayerPrimaryGameServerIdx = ProxyPlayerToPrimaryGameServer.Find(LocalPlayer);
-			if (PlayerPrimaryGameServerIdx && *PlayerPrimaryGameServerIdx < GameServerConnections.Num())
+			if (ensure(*PlayerPrimaryGameServerIdx < GameServerConnections.Num()))
 			{
 				FGameServerConnectionState& PlayerPrimaryGameServer = GameServerConnections[*PlayerPrimaryGameServerIdx];
 				if (PlayerPrimaryGameServer.NetDriver->ServerConnection == ParentGameServerConnection)
 				{
-					UE_LOG(LogNetProxy, Log, TEXT("Binding proxy connection %s:%s to game server connection %s:%s (controller=%s player=%s:%d)"),
-						   *ProxyConnection->GetDriver()->GetName(),
-						   *ProxyConnection->GetName(),
-						   *ChildGameServerConnection->GetDriver()->GetName(),
-						   *ChildGameServerConnection->GetName(),
-						   *PlayerController->GetName(),
-						   *LocalPlayer->GetName(),
-						   PlayerIndex);
-
-					// Called here because this is the first time we have a pointer to the child connection.
-					PrepareGameServerConnection(ChildGameServerConnection);
-
-					// Ensure that the proxy connection always uses the local player associated with this game connection.
-					ProxyConnection->SetLocalPlayerFunction([LocalPlayer](
-						UNetConnection* GameServerConnection,
-						APlayerController* PlayerController)
-					{
-						return LocalPlayer;
-					});
+					bBindProxyConnectionToPrimaryGameServer = true;
+				}
+			}
+		}
 
-					// Complete the proxy connection handshake using the player controller replicated from the game server.
-					ProxyConnection->HandleClientPlayer(PlayerController, ProxyConnection);
+		// If the proxy connection has already been bound to a primary game server player controller, and the new player has the 
+		// same player id as the existing player id, then change the primary game server and rebind the two connections.
+		/*else if(ProxyConnection->PlayerController)
+		{
+			if (ULocalPlayer* ExistingPlayer = Cast<ULocalPlayer>(ProxyConnection->PlayerController->Player))
+			{
+				if (ProxyPlayerToPrimaryGameServerPlayerController.Contains(ExistingPlayer))
+				{
+					int32* PreviousPrimaryGameServerIdx = ProxyPlayerToPrimaryGameServer.Find(ExistingPlayer);
+					if (ensure(PreviousPrimaryGameServerIdx))
+					{
+						if (NewPlayer->GetCachedUniqueNetId() == ExistingPlayer->GetCachedUniqueNetId())
+						{
+							int32 NewPrimaryGameServerIdx = GetGameServerWithPlayer(NewPlayer);
+
+							UE_LOG(LogNetProxy, Log, TEXT("Changing primary game server for proxy connection %s:%s from player %s to %s (prev_game_server_idx=%d new_game_server_idx=%d)."),
+								   *ProxyConnection->Driver->GetName(),
+								   *ProxyConnection->GetName(),
+								   *ExistingPlayer->GetName(),
+								   *NewPlayer->GetName(),
+								   *PreviousPrimaryGameServerIdx,
+								   NewPrimaryGameServerIdx);
+					
+							ProxyPlayerToPrimaryGameServerPlayerController.Remove(ExistingPlayer);
+
+							DeregisterPrimaryGameServer(ProxyConnection);
+							RegisterPrimaryGameServer(ProxyConnection, NewPlayer, NewPrimaryGameServerIdx);
+				
+							bBindProxyConnectionToPrimaryGameServer = true;
+						}
+					}
 				}
 			}
+		}*/
+
+		if (bBindProxyConnectionToPrimaryGameServer)
+		{
+			UPlayer* PreviousPlayer = GameServerPlayerController->Player;
 			
-			// Remap the proxy connection to the new child connection.
-			FGameServerAndPlayerIndex NewKey; 
-			NewKey.GameServerConnection = ChildGameServerConnection; 
-			NewKey.PlayerIndex = PlayerIndex;
-			GameToProxyConnection.Remove(Key);
-			GameToProxyConnection.Add(NewKey, ProxyConnection);
+			APlayerController* ProxyPlayerController = SpawnPlayerController();
+
+			if (ProxyPlayerController)
+			{
+				DetachAndDestroyPlayerController(ProxyConnection->PlayerController);
+				AttachPlayerController(ProxyPlayerController, ProxyConnection, NewPlayer);
+
+				// The player always points to the proxy player controller if that player is
+				// considered the primary player for the proxy connection.
+				NewPlayer->PlayerController = ProxyPlayerController;
+
+				ProxyPlayerToPrimaryGameServerPlayerController.Add(NewPlayer, GameServerPlayerController);
+
+				UE_LOG(LogNetProxy, Log, TEXT("Binding proxy connection %s:%s to primary game server connection %s:%s (game_server_controller=%s proxy_controller=%s player=%s:%d)"),
+					   *ProxyConnection->GetDriver()->GetName(),
+					   *ProxyConnection->GetName(),
+					   *ChildGameServerConnection->GetDriver()->GetName(),
+					   *ChildGameServerConnection->GetName(),
+					   *GameServerPlayerController->GetName(),
+					   *ProxyPlayerController->GetName(),
+					   *NewPlayer->GetName(),
+					   PlayerIndex);
+			}
+			else
+			{
+				UE_LOG(LogNetProxy, Error, TEXT("Failed to bind a proxy connection %s:%s to a new primary game server %s:%s."),
+					   *ProxyConnection->GetDriver()->GetName(),
+					   *ProxyConnection->GetName(),
+					   *ChildGameServerConnection->GetDriver()->GetName(),
+					   *ChildGameServerConnection->GetName());
+			}
+		}
+		else
+		{
+			UE_LOG(LogNetProxy, Log, TEXT("Ignoring connection %s:%s to non-primary game server %s:%s (game_server_controller=%s player=%s:%d)"),
+				*ProxyConnection->GetDriver()->GetName(),
+				*ProxyConnection->GetName(),
+				*ChildGameServerConnection->GetDriver()->GetName(),
+				*ChildGameServerConnection->GetName(),
+				*GameServerPlayerController->GetName(),
+				*NewPlayer->GetName(),
+				PlayerIndex);
 		}
+			
+		// Update the mapping of proxy connection to game server connection to use the new child game server connection (if it's 
+		// different from the parent connection).
+		FGameServerAndPlayerIndex NewKey; 
+		NewKey.GameServerConnection = ChildGameServerConnection; 
+		NewKey.PlayerIndex = PlayerIndex;
+		GameToProxyConnection.Remove(Key);
+		GameToProxyConnection.Add(NewKey, ProxyConnection);
 	}
 
 	// If this was a parent connection send any join requests for any pending multiplexed connections.
@@ -568,4 +813,207 @@ void UProxyNetDriver::FlushSplitJoinRequests(FGameServerConnectionState* GameSer
 	}
 
 	GameServerConnectionState->PendingSplitJoinRequests.Reset();
+}
+
+int32 UProxyNetDriver::GetGameServerWithPlayer(ULocalPlayer* Player)
+{
+	for (int32 GameServerIdx = 0; GameServerIdx < GameServerConnections.Num(); GameServerIdx++)
+	{
+		if (GameServerConnections[GameServerIdx].Players.Find(Player) != INDEX_NONE)
+		{
+			return GameServerIdx;
+		}
+	}
+
+	return -1;
+}
+
+void UProxyNetDriver::DeregisterPrimaryGameServer(UNetConnection* ProxyConnection)
+{
+	if (ensure(ProxyConnection->PlayerController))
+	{
+		ULocalPlayer* OldPlayer = Cast<ULocalPlayer>(ProxyConnection->PlayerController->Player);
+		if (OldPlayer)
+		{
+			FUniqueNetIdRepl& PlayerId = ProxyConnection->PlayerId;
+
+			UE_LOG(LogNetProxy, Log, TEXT("Clearing player %s primary game server (proxy_connection=%s:%s)"),
+				   *OldPlayer->GetName(),
+				   *ProxyConnection->GetDriver()->GetName(),
+				   *ProxyConnection->GetName());
+
+			ensure(ProxyPlayerToPrimaryGameServer.Contains(OldPlayer));
+		
+			ProxyPlayerToPrimaryGameServer.Remove(OldPlayer);
+		}
+	}
+}
+
+void UProxyNetDriver::RegisterPrimaryGameServer(UNetConnection* ProxyConnection, class ULocalPlayer* PrimaryPlayer, int32 GameServerConnectionStateIndex)
+{
+	if (ensure(GameServerConnectionStateIndex < GameServerConnections.Num()))
+	{
+		FUniqueNetIdRepl& PlayerId = ProxyConnection->PlayerId;
+
+		UE_LOG(LogNetProxy, Log, TEXT("Configuring primary game server %d for player (player=%s proxy_connection=%s:%s)"),
+			GameServerConnectionStateIndex,
+			*PrimaryPlayer->GetName(),
+			*ProxyConnection->GetDriver()->GetName(),
+			*ProxyConnection->GetName());
+
+		ensure(!ProxyPlayerToPrimaryGameServer.Contains(PrimaryPlayer));
+
+		ProxyPlayerToPrimaryGameServer.Add(PrimaryPlayer, GameServerConnectionStateIndex);
+	}
+}
+
+void UProxyNetDriver::AttachPlayerController(APlayerController* PlayerController, UNetConnection* Connection, ULocalPlayer* Player)
+{
+	// The presumed relationships between connection, player controller and player:
+	//	
+	//	ProxyConnection <-> ProxyPlayerController -> LocalPlayer
+	//	LocalPlayer <- GameServerPlayerController <-> Game Server Connection
+	//
+	// This function doesn't associate the player with a player controller since it differs between
+	// proxy and game server player controllers.
+
+	Connection->SetConnectionState(EConnectionState::USOCK_Open);
+
+	// It's assumed the connection is not associated with a player controller.	
+	ensure(Connection->PlayerController == nullptr);
+	ensure(Connection->OwningActor == nullptr);
+
+	Connection->PlayerController = PlayerController;
+	Connection->OwningActor = PlayerController;
+	
+	PlayerController->NetConnection = Connection;
+	PlayerController->Player = Player;
+	
+	Connection->LastReceiveTime = GetElapsedTime();
+}
+
+void UProxyNetDriver::DetachPlayerController(APlayerController* PlayerController)
+{
+	if (PlayerController)
+	{
+		UPlayer* Player = PlayerController->Player;
+
+		UE_LOG(LogNetProxy, Log, TEXT("Destroying old controller %s on %s attached to player %s."),
+			*PlayerController->GetName(),
+			*this->GetName(),
+			*Player->GetName());
+
+		if (UNetConnection* Connection = PlayerController->NetConnection)
+		{
+			Connection->PlayerController = nullptr;
+			Connection->OwningActor = nullptr;
+		}
+
+		PlayerController->NetConnection = nullptr;
+		PlayerController->Player = nullptr;
+
+		// Only detach the player controller from the player if it's the one referenced by the player.
+		if (Player->PlayerController == PlayerController)
+		{
+			Player->PlayerController = nullptr;
+		}
+	}
+}
+
+void UProxyNetDriver::DetachAndDestroyPlayerController(APlayerController* PlayerController)
+{
+	DetachPlayerController(PlayerController);
+	DestroyPlayerController(PlayerController);
+}
+
+APlayerController* UProxyNetDriver::SpawnPlayerController()
+{
+	AGameModeBase* GameMode = GetWorld()->GetAuthGameMode();
+
+	if (ensure(GameMode))
+	{
+		FActorSpawnParameters SpawnInfo;
+		SpawnInfo.ObjectFlags |= RF_Transient;
+		SpawnInfo.bDeferConstruction = true;
+	
+		APlayerController* PlayerController = GetWorld()->SpawnActor<APlayerController>(GameMode->PlayerControllerClass, SpawnInfo);
+
+		if (PlayerController)
+		{
+			UE_LOG(LogNetProxy, Log, TEXT("Spawning player controller %s of type %s."), *PlayerController->GetName(), *GameMode->PlayerControllerClass->GetName());
+
+			ProxySpawnedActors.Add(PlayerController);
+	
+			GetWorld()->AddController(PlayerController);
+
+			PlayerController->SetRole(ROLE_Authority);
+			PlayerController->SetReplicates(true);
+			PlayerController->SetAutonomousProxy(true); // Sets AActor::RemoteRole.
+			PlayerController->FinishSpawning(FTransform());
+		}
+
+		return PlayerController;
+	}
+
+	return nullptr;
+}
+
+void UProxyNetDriver::DestroyPlayerController(APlayerController* PlayerController)
+{
+	if (PlayerController)
+	{
+		// Only player controllers spawned by the proxy should be destroyed.
+		if (ensure(ProxySpawnedActors.Contains(PlayerController)))
+		{
+			if (ensure(PlayerController->GetLocalRole() == ROLE_Authority))
+			{
+				ProxySpawnedActors.Remove(PlayerController);
+
+				PlayerController->GetWorld()->RemoveController(PlayerController);
+				PlayerController->GetWorld()->DestroyActor(PlayerController);
+			}
+		}
+	}
+}
+
+ULocalPlayer* UProxyNetDriver::GetPlayerOnGameServer(UNetConnection* ParentGameServerConnection, int32 PlayerIndex)
+{
+	ULocalPlayer* FoundPlayer = nullptr;
+
+	for (FGameServerConnectionState& GameServerConnectionState : GameServerConnections)
+	{
+		if (GameServerConnectionState.NetDriver == ParentGameServerConnection->GetDriver())
+		{
+			if (ensure(PlayerIndex < GameServerConnectionState.Players.Num()))
+			{
+				FoundPlayer = GameServerConnectionState.Players[PlayerIndex];
+			}
+		}
+	}
+
+	return FoundPlayer;
+}
+
+UNetConnection* UProxyNetDriver::GetProxyConnectionFromPrimaryPlayer(UPlayer* Player)
+{
+	if (ULocalPlayer* LocalPlayer = Cast<ULocalPlayer>(Player))
+	{
+		if (APlayerController* ProxyPlayerController = LocalPlayer->PlayerController)
+		{
+			return ProxyPlayerController->NetConnection;
+		}
+	}
+
+	return nullptr;
+}
+
+APlayerController* UProxyNetDriver::GetProxyControllerFromPrimaryGameServerController(APlayerController* GameServerController)
+{
+	if (ULocalPlayer* Player = Cast<ULocalPlayer>(GameServerController->Player))
+	{
+		// The player's controller will always point to the proxy player controller.
+		return Player->PlayerController;
+	}
+
+	return nullptr;
 }
\ No newline at end of file
diff --git a/Engine/Plugins/Runtime/MultiServerReplication/Source/MultiServerReplication/Public/MultiServerProxy.h b/Engine/Plugins/Runtime/MultiServerReplication/Source/MultiServerReplication/Public/MultiServerProxy.h
--- a/Engine/Plugins/Runtime/MultiServerReplication/Source/MultiServerReplication/Public/MultiServerProxy.h
+++ b/Engine/Plugins/Runtime/MultiServerReplication/Source/MultiServerReplication/Public/MultiServerProxy.h
@@ -3,7 +3,10 @@
 #pragma once
 
 #include "Engine/PendingNetGame.h"
+#include "Engine/ChildConnection.h"
+#include "GameFramework/OnlineReplStructs.h"
 #include "IpNetDriver.h"
+#include "IpConnection.h"
 #include "MultiServerProxy.generated.h"
 
 USTRUCT()
@@ -124,6 +127,32 @@ private:
 	TObjectPtr<class UProxyNetDriver> ProxyNetDriver;
 };
 
+/** 
+* A network connection used by UProxyBackendNetDriver.
+ */
+UCLASS()
+class MULTISERVERREPLICATION_API UProxyBackendNetConnection : public UIpConnection
+{
+	GENERATED_BODY()
+
+public:
+
+	virtual void HandleClientPlayer(APlayerController* PC, UNetConnection* NetConnection);
+};
+
+/**
+ * A child network connection used by UProxyBackendNetDriver.
+ */
+UCLASS()
+class MULTISERVERREPLICATION_API UProxyBackendChildNetConnection : public UChildConnection
+{
+	GENERATED_BODY()
+
+public:
+
+	virtual void HandleClientPlayer(APlayerController* PC, UNetConnection* NetConnection);
+};
+
 /** 
  * A driver that is used by UProxyNetDriver to connect to backend game servers.
  */
@@ -135,9 +164,13 @@ class MULTISERVERREPLICATION_API UProxyBackendNetDriver : public UIpNetDriver
 public:
 
 	void SetProxyNetDriver(TObjectPtr<UProxyNetDriver> InParentNetDriver);
+	TObjectPtr<UProxyNetDriver> GetProxyNetDriver();
 
 	virtual bool InitBase(bool bInitAsClient, FNetworkNotify* InNotify, const FURL& URL, bool bReuseAddressAndPort, FString& Error);
 	virtual void ForwardRemoteFunction(UObject* RootObject, UObject* SubObject, UFunction* Function, void* Parms) override;
+	virtual bool ShouldSkipRepNotifies() const override;
+	virtual UChildConnection* CreateChild(UNetConnection* Parent) override;
+	virtual void InternalProcessRemoteFunction(class AActor* Actor, class UObject* SubObject, class UNetConnection* Connection, class UFunction* Function, void* Parms, FOutParmRec* OutParms, FFrame* Stack, bool bIsServer) override;
 
 private:
 
@@ -146,31 +179,46 @@ private:
 };
 
 /**
-* A network proxy that intercepts and forwards UE game network connections to game servers.
-*
-* This proxy externally behaves the same as a normal game server which means any game client
-* can connect to it without any special configuration.
+* A network proxy that intercepts and forwards UE game network connections to backend game servers.
 *
-* Internally, the proxy is made up of a NetDriver that listens for incoming connections known
-* as proxy connections, and a NetDriver for each connection to a backend game server. State from
-* the backend servers is replicated into a single, shared UWorld and the listening NetDriver 
-* replicates that state out to proxy connections.
+* The proxy externally behaves the same as a normal game server when game clients connect and
+* as a normal client when connecting to game servers. This means that there is no need for the
+* clients and game servers that the proxy is connected to have any special proxy-aware configuration.
 *
-* When a client connects to the proxy, the proxy establishes a connection to each backend game
-* server. This means that the proxy will hold an instance of ULocalPlayer and APlayerController 
-* for each client connected to the proxy multiplied by each game server (e.g. 2 clients and 2 
-* game servers would result in 4 instances of APlayerController and ULocalPlayer).
-* 
-* The instances of APlayerController set their member variable APlayerController::NetConnection
-* to a proxy connection. This means that calls to APlayerController::GetNetConnection() will 
-* always be to a proxy connection that lead to the instantiation of that controller.
-*
-* A player connected to the proxy considers one of the backend servers to be it's primary 
-* game server. This allows the proxy to preserve the idea that a player is logically connected
-* to only one game server (e.g. a player controller can only possess a single pawn at a time).
+* Internally, the proxy is made up of an instance of UProxyNetDriver that listens for incoming 
+* connections, known as proxy connections, and an instance of UProxyBackendNetDriver for each 
+* connection to a backend game server. State from the backend servers is replicated into a single,
+* shared UWorld and the listening UProxyNetDriver replicates that state out to proxy connections.
 *
 * All actors replicated to the proxy from remote game servers will have the same role as a client
 * (ROLE_SimulatedProxy or ROLE_AutonomousProxy) and will be replicated as-is to the proxy client.
+*
+* When a proxy connection (UNetConnection) is opened in UProxyNetDriver it opens a game server 
+* connection (UNetConnection) to each registered backend server. Each proxy connection and game 
+* server connection is associated with it's own instance of APlayerController. There will always 
+* be an equal number of connections and player controllers on a proxy and can be calculated with
+* this equation: i+(j*i), where i is the number of clients connected to the proxy and j is the 
+* number of registered game servers.
+*
+* The proxy and game server player controllers are related through an instance of ULocalPlayer
+* for each proxy connection to a registered game server: the equation for the number of players
+* on the proxy is (i*j), where i is the number of clients connected to the proxy and j in the 
+* number of registered game servers.
+*
+* The relationship between proxy connections and game server connections is done through pointers
+* in the respective classes:
+*
+* ProxyConnection <-> ProxyPlayerController <-> LocalPlayer <- GameServerPlayerController <-> GameServerConnection
+*
+* From this representation of the relationship we can see that ProxyConnection has a pointer to 
+* ProxyPlayerController, and ProxyPlayerController has a pointer to ProxyConnection. Using these
+* existing pointers enables the use of existing UE code and avoids the need for additional data
+* structures to hold the relationship betweeen the clients connected to the proxy
+* and registered game servers.
+*
+* IMPORTANT: There isn't a relationship from LocalPlayer to GameServerPlayerController because
+* the exiting pointer (ULocalPlayer::PlayerController) that is being used can only point to one
+* controller.
 */
 UCLASS()
 class MULTISERVERREPLICATION_API UProxyNetDriver : public UIpNetDriver
@@ -189,18 +237,28 @@ public:
 	int32 GetGameServerConnectionCount() const;
 	FGameServerConnectionState* GetGameServerConnection(int32 Index);
 
+	/** UNetDriver interface function. */
 	virtual bool InitBase(bool bInitAsClient, FNetworkNotify* InNotify, const FURL& URL, bool bReuseAddressAndPort, FString& Error) override;
 	virtual bool InitConnect(class FNetworkNotify* InNotify, const FURL& ConnectURL, FString& Error) override;
 	virtual void ForwardRemoteFunction(UObject* RootObject, UObject* SubObject, UFunction* Function, void* Parms) override;
 	virtual bool ShouldReplicateFunction(AActor* Actor, UFunction* Function) const override;
 	virtual void NotifyActorChannelOpen(UActorChannel* Channel, AActor* Actor) override;
 	virtual void AddNetworkActor(AActor* Actor) override;
+	virtual bool ShouldCallRemoteFunction(UObject* Object, UFunction* Function, const FReplicationFlags& RepFlags) const;
+	virtual void InternalProcessRemoteFunction(class AActor* Actor, class UObject* SubObject, class UNetConnection* Connection, class UFunction* Function, void* Parms, FOutParmRec* OutParms, FFrame* Stack, bool bIsServer) override;
+	virtual int32 ServerReplicateActors(float DeltaSeconds) override;
+	virtual bool CanDowngradeActorRole(UNetConnection* ProxyConnection, AActor* Actor) const override;
 
 	/** Get the next identifier for incoming client connections to the proxy. */
 	int32 GetNextProxyClientId();
 	
 	/** Get the next identifier for outgoing connections to game servers. */
 	int32 GetNextGameServerClientId();
+	
+	/** Called when a parent or child game server connection has been established. */
+	void GameServerConnectionHandshakeComplete(UNetConnection* ChildGameServerConnection,
+											   APlayerController* PlayerController,
+											   UNetConnection* ParentGameServerConnection);
 
 private:
 
@@ -209,40 +267,60 @@ private:
 	// Classes in EngineTest that test the behavior of UMultiServerProxy.
 	friend class FNetTestProxyRPCRouting;
 	friend class FNetTestProxyConnectionRouting;
+	friend class FNetTestProxyMovePlayerController;
 	
 	friend class UProxyBackendNetDriver;
+
+	/** Return true if an actor is explicitly spawned by the proxy. */
+	bool IsProxySpawned(AActor* Actor) const;
 	
-	/** Return the local player associated with a connection to a game server. */
-	class ULocalPlayer* GetPlayer(UNetConnection* GameServerConnection, APlayerController* PlayerController);
+	/** Send any split join requests that have been queued up to the game server. */
+	void FlushSplitJoinRequests(FGameServerConnectionState* GameServerConnectionState);
 
-	/** Perform common configuration of a connection to a game server. */
-	void PrepareGameServerConnection(UNetConnection* GameServerConnection);
+	/** Return a game server index that a player is connected to. */
+	int32 GetGameServerWithPlayer(class ULocalPlayer* Player);
 
-	/** Called when a game server connection has been established. */
-	void ConnectionHandshakeComplete(UNetConnection* ChildGameServerConnection,
-									 APlayerController* PlayerController,
-									 UNetConnection* ParentGameServerConnection);
+	/** Disassociate a proxy connection with a game server as a primary game server. */
+	void DeregisterPrimaryGameServer(UNetConnection* ProxyConnection);
 
-	/** Send any split join requests that have been queued up to the game server. */
-	void FlushSplitJoinRequests(FGameServerConnectionState* GameServerConnectionState);
+	/** Associate a proxy connection with a primary game server. */
+	void RegisterPrimaryGameServer(UNetConnection* ProxyConnection, class ULocalPlayer* PrimaryPlayer, int32 GameServerConnectionStateIndex);
 
-	/** Map a game server and player controller index to a proxy connection. */
+	/** Attach a new player, player controller and connection. */
+	void AttachPlayerController(APlayerController* PlayerController, UNetConnection* Connection, ULocalPlayer* Player);
+
+	/** Detach an existing player, player controller and connection and (optionally) destroy the controller. */
+	void DetachPlayerController(APlayerController* PlayerController);
+	void DetachAndDestroyPlayerController(APlayerController* PlayerController);
+
+	/** Create and destroy a proxy player controller. Returns null if failed to spawn the controller. */
+	APlayerController* SpawnPlayerController();
+
+	/** Destroy a player controller spawned by the proxy. */
+	void DestroyPlayerController(APlayerController* PlayerController);
+	
+	/** Return the local player associated with a player index on a given game server. */
+	class ULocalPlayer* GetPlayerOnGameServer(UNetConnection* ParentGameServerConnection, int32 PlayerIndex);
+
+	/** Return a proxy connection from a primary player (a player on a primary game server). */
+	static UNetConnection* GetProxyConnectionFromPrimaryPlayer(UPlayer* Player);
+
+	/** Return a proxy player controller given a primary game server controller. */
+	APlayerController* GetProxyControllerFromPrimaryGameServerController(APlayerController* GameServerPlayerController);
+	
+	/** Associate a connection to a game server with the proxy connection that established that connection. */
 	UPROPERTY()
 	TMap<FGameServerAndPlayerIndex, TObjectPtr<UNetConnection>> GameToProxyConnection;
 
-	/** Map a proxy connection to a primary proxy player. */
+	/** Map a player to it's primary game server index into GameServerConnections (set when connection handshake begins). */
 	UPROPERTY()
-	TMap<TObjectPtr<UNetConnection>, TObjectPtr<class ULocalPlayer>> ProxyConnectionToPrimaryPlayer;
+	TMap<TObjectPtr<class ULocalPlayer>, int32> ProxyPlayerToPrimaryGameServer;
 
-	/** Map a proxy player to it's primary game server index. */
+	/** Map a player to it's primary game server player controller (set after connection handshake is completed successfully). */
 	UPROPERTY()
-	TMap<TObjectPtr<class ULocalPlayer>, int32> ProxyPlayerToPrimaryGameServer;
+	TMap<TObjectPtr<class ULocalPlayer>, TObjectPtr<APlayerController>> ProxyPlayerToPrimaryGameServerPlayerController;
 
-	/** 
-	 * Net drivers and associated state used to connect to backend game servers.
-	 *
-	 * NOTE: This will need to be replaced with a TMap or similar to handle disconnects to the game servers.
-	 */
+	/** Net drivers and associated state used to connect to backend game servers. */
 	UPROPERTY()
 	TArray<FGameServerConnectionState> GameServerConnections;
 
@@ -250,5 +328,9 @@ private:
 	UPROPERTY()
 	TObjectPtr<UProxyListenerNotify> ProxyListenerNotify;
 
+	/** A set of actors that have been spawned by the proxy. */
+	UPROPERTY()
+	TSet<TObjectPtr<AActor>> ProxySpawnedActors;
+
 	int32 NextGameServerClientId = 0;
 };
diff --git a/Engine/Source/Runtime/Engine/Classes/Engine/ChildConnection.h b/Engine/Source/Runtime/Engine/Classes/Engine/ChildConnection.h
--- a/Engine/Source/Runtime/Engine/Classes/Engine/ChildConnection.h
+++ b/Engine/Source/Runtime/Engine/Classes/Engine/ChildConnection.h
@@ -82,14 +82,14 @@ public:
 	{
 		SetConnectionState(Parent->GetConnectionState());
 	}
-	virtual void GracefulClose(FNetCloseResult&& CloseReason) override;
-	virtual void HandleClientPlayer(class APlayerController* PC, class UNetConnection* NetConnection) override;
-	virtual void CleanUp() override;
+	ENGINE_API virtual void GracefulClose(FNetCloseResult&& CloseReason) override;
+	ENGINE_API virtual void HandleClientPlayer(class APlayerController* PC, class UNetConnection* NetConnection) override;
+	ENGINE_API virtual void CleanUp() override;
 
 	virtual void InitRemoteConnection(UNetDriver* InDriver, class FSocket* InSocket, const FURL& InURL, const class FInternetAddr& InRemoteAddr, EConnectionState InState, int32 InMaxPacket = 0, int32 InPacketOverhead = 0) override {}
 	virtual void InitLocalConnection(UNetDriver* InDriver, class FSocket* InSocket, const FURL& InURL, EConnectionState InState, int32 InMaxPacket = 0, int32 InPacketOverhead = 0) override {}
 
-	virtual void InitChildConnection(UNetDriver* InDriver, UNetConnection* InParent);
+	ENGINE_API virtual void InitChildConnection(UNetDriver* InDriver, UNetConnection* InParent);
 
 private:
 	void AssignConnectionHandle();
diff --git a/Engine/Source/Runtime/Engine/Classes/Engine/NetDriver.h b/Engine/Source/Runtime/Engine/Classes/Engine/NetDriver.h
--- a/Engine/Source/Runtime/Engine/Classes/Engine/NetDriver.h
+++ b/Engine/Source/Runtime/Engine/Classes/Engine/NetDriver.h
@@ -799,7 +799,7 @@ class UNetDriver : public UObject, public FExec
 
 protected:
 
-	ENGINE_API void InternalProcessRemoteFunction(
+	virtual ENGINE_API void InternalProcessRemoteFunction(
 		class AActor* Actor,
 		class UObject* SubObject,
 		class UNetConnection* Connection,
@@ -1471,14 +1471,17 @@ protected:
 
 public:
 
-	/** Enable/disable the execution of received RPC functions. */
+	/** (Generic Replicaton Only) Enable/disable the execution of received RPC functions. */
 	ENGINE_API void EnableExecuteRPCFunctions(bool bEnable);
 	ENGINE_API bool IsExecuteRPCFunctionsEnabled() const;
 
-	/** Treat actors with ROLE_AutonomousProxy as ROLE_SimulatedProxy when replicating state (applies to all client connections). */
+	/** (Generic Replicaton Only) Treat actors with ROLE_AutonomousProxy as ROLE_SimulatedProxy when replicating state (applies to all client connections). */
 	ENGINE_API void ReplicateAutonomousAsSimulated(bool bEnabled);
 	ENGINE_API bool IsReplicatingAutonomousAsSimulated() const;
 
+	/** (Generic Replicaton Only) Return true if an actor's role can be downgraded during replication to a given connection.  */
+	virtual bool CanDowngradeActorRole(UNetConnection* Connection, AActor* Actor) const { return true; };
+
 private:
 
 	bool bExecuteRPCFunctions = true;
diff --git a/Engine/Source/Runtime/Engine/Private/DataChannel.cpp b/Engine/Source/Runtime/Engine/Private/DataChannel.cpp
--- a/Engine/Source/Runtime/Engine/Private/DataChannel.cpp
+++ b/Engine/Source/Runtime/Engine/Private/DataChannel.cpp
@@ -3759,7 +3759,11 @@ int64 UActorChannel::ReplicateActor()
 	}
 
 	// Possibly downgrade role of actor if this connection doesn't own it
-	FScopedRoleDowngrade ScopedRoleDowngrade( Actor, RepFlags );
+	TUniquePtr<FScopedRoleDowngrade> ScopedRoleDowngrade;
+	if (Connection->Driver->CanDowngradeActorRole(Connection, Actor))
+	{
+		ScopedRoleDowngrade = MakeUnique<FScopedRoleDowngrade>(Actor, RepFlags);
+	}
 
 	RepFlags.bNetSimulated	= (Actor->GetRemoteRole() == ROLE_SimulatedProxy);
 
