Support for routing RPCs to actor components through the proxy.

diff --git a/Engine/Plugins/Runtime/MultiServerReplication/Source/MultiServerReplication/Private/MultiServerProxy.cpp b/Engine/Plugins/Runtime/MultiServerReplication/Source/MultiServerReplication/Private/MultiServerProxy.cpp
--- a/Engine/Plugins/Runtime/MultiServerReplication/Source/MultiServerReplication/Private/MultiServerProxy.cpp
+++ b/Engine/Plugins/Runtime/MultiServerReplication/Source/MultiServerReplication/Private/MultiServerProxy.cpp
@@ -323,39 +323,113 @@ void UProxyBackendNetDriver::InternalProcessRemoteFunction(
 	FFrame* Stack,
 	bool bIsServer)
 {
-	AActor* DestinationActor = Actor;
+	bool bShouldForwardRPC = true;
 
 	// The majority of the actors on the proxy are replicated from the game servers so their owners will be game server player
 	// controllers. The exception are proxy player controllers which are spawned by the proxy and will have an owner that is
 	// the proxy player controller. In this case we must forward the RPC to the primary game server player controller.
 	if (ProxyNetDriver->IsProxySpawned(Actor))
 	{
-		if (APlayerController* ProxyPlayerController = Cast<APlayerController>(Actor))
+		AActor* MappedActor = Actor;
+
+		if (APlayerController* ProxyPlayerController = Cast<APlayerController>(MappedActor))
 		{
 			if (ULocalPlayer* Player = Cast<ULocalPlayer>(ProxyPlayerController->Player))
 			{
 				if (TObjectPtr<APlayerController>* GameServerPlayerControllerPtr = ProxyNetDriver->ProxyPlayerToPrimaryGameServerPlayerController.Find(Player))
 				{
 					UNetConnection* GameServerConnection = GameServerPlayerControllerPtr->Get()->NetConnection;
-					ensure(GameServerConnection->PlayerController->Player == Player);
-					DestinationActor = GameServerConnection->PlayerController;
-					SubObject = nullptr;
+					if (ensure(GameServerConnection->PlayerController))
+					{
+						ensure(GameServerConnection->PlayerController->Player == Player);
+						MappedActor = GameServerConnection->PlayerController;
+
+						UE_LOG(LogNetProxy, VeryVerbose, TEXT("RPC: Remapping proxy player controller %s to game server player controller %s when forwarding RPC %s on %s"),
+							   *ProxyPlayerController->GetName(),
+							   *GameServerConnection->PlayerController->GetName(),
+							   *Function->GetName(),
+							   *Actor->GetName());
+					}
+				}
+			}
+		}
+		
+		if (ProxyNetDriver->IsProxySpawned(MappedActor))
+		{
+			bShouldForwardRPC = false;
 
-					UE_LOG(LogNetProxy, VeryVerbose, TEXT("RPC: Remapping proxy player controller %s to game server player controller %s when forwarding RPC %s on %s"), 
-						   *ProxyPlayerController->GetName(), *GameServerConnection->PlayerController->GetName(), *Function->GetName(), *DestinationActor->GetName());
+			UE_LOG(LogNetProxy, Warning, TEXT("RPC: Unable to map proxy owned actor %s to game server %s owned actor."), *Actor->GetName(), *Connection->GetName());
+		}
+		else
+		{
+			Actor = MappedActor;
+		}
+	}
+
+	// If the sub-object is not owned by the actor, attempt to find an component in that actor that matches the same type.
+	// This logic assumes that an actor only has one component of a given type and will fail if that assumption is incorrect.
+	if (UActorComponent* SubObjectAsActorComponent = Cast<UActorComponent>(SubObject))
+	{
+		if (!Actor->OwnsComponent(SubObjectAsActorComponent))
+		{
+			int32 MatchingComponents = 0;
+			UActorComponent* MappedActorComponent = SubObjectAsActorComponent;
+
+			for (UActorComponent* ActorComponent : Actor->GetComponents())
+			{
+				if (ActorComponent->GetClass() == SubObjectAsActorComponent->GetClass())
+				{
+					MappedActorComponent = ActorComponent;
+					MatchingComponents++;
+
+					UE_LOG(LogNetProxy, VeryVerbose, TEXT("RPC: Remapping sub-object %s to sub-object %s in actor %s"),
+						*SubObjectAsActorComponent->GetName(),
+						*MappedActorComponent ->GetName(),
+						*Actor->GetName());
 				}
 			}
+
+			if (!Actor->OwnsComponent(MappedActorComponent))
+			{
+				bShouldForwardRPC = false;
+
+				UE_LOG(LogNetProxy, Warning, TEXT("RPC: Unable map sub-object %s to actor %s"),
+					   *SubObjectAsActorComponent->GetName(),
+					   *Actor->GetName());
+			}
+
+			// Detect an actor with two components of the same type.
+			else if (MatchingComponents > 1)
+			{
+				bShouldForwardRPC = false;
+
+				UE_LOG(LogNetProxy, Warning, TEXT("RPC: Found an actor %s with more than one component %s."), *Actor->GetName(), *SubObject->GetName());
+			}
+			else
+			{
+				SubObject = MappedActorComponent;
+			}
 		}
 	}
 
-	if (!ProxyNetDriver->IsProxySpawned(DestinationActor))
+	if (bShouldForwardRPC)
 	{
-		UE_LOG(LogNetProxy, Verbose, TEXT("RPC: Forwarding %s on %s to game server connection %s"), *Function->GetName(), *DestinationActor->GetName(), *Connection->GetName());
-		Super::InternalProcessRemoteFunction(DestinationActor, SubObject, Connection, Function, Parms, OutParms, Stack, bIsServer);
+		UE_LOG(LogNetProxy, Verbose, TEXT("RPC: Forwarding %s on %s (owner:%s sub-object:%s) to game server connection %s"), 
+			   *Function->GetName(), 
+			   *Actor->GetName(), 
+			   *GetNameSafe(Actor->GetOwner()),
+			   *GetNameSafe(SubObject),
+			   *Connection->GetName());
+
+		Super::InternalProcessRemoteFunction(Actor, SubObject, Connection, Function, Parms, OutParms, Stack, bIsServer);
 	}
 	else
 	{
-		UE_LOG(LogNetProxy, Verbose, TEXT("RPC: Ignoring %s on %s to game server connection %s because it's not owned by a game server."), *Function->GetName(), *DestinationActor->GetName(), *Connection->GetName());
+		UE_LOG(LogNetProxy, Verbose, TEXT("RPC: Ignoring %s on %s (sub-object:%s) to game server connection %s."), 
+			   *Function->GetName(), 
+			   *Actor->GetName(), 
+			   *GetNameSafe(SubObject),
+			   *Connection->GetName());
 	}
 }
 
