Fix for a compile error in MultiServerProxy.cpp.

diff --git a/Engine/Plugins/Runtime/MultiServerReplication/Source/MultiServerReplication/Private/MultiServerProxy.cpp b/Engine/Plugins/Runtime/MultiServerReplication/Source/MultiServerReplication/Private/MultiServerProxy.cpp
--- a/Engine/Plugins/Runtime/MultiServerReplication/Source/MultiServerReplication/Private/MultiServerProxy.cpp
+++ b/Engine/Plugins/Runtime/MultiServerReplication/Source/MultiServerReplication/Private/MultiServerProxy.cpp
@@ -370,7 +370,7 @@ void UProxyBackendNetDriver::InternalProcessRemoteFunction(
 	// This logic assumes that an actor only has one component of a given type and will fail if that assumption is incorrect.
 	if (UActorComponent* SubObjectAsActorComponent = Cast<UActorComponent>(SubObject))
 	{
-		if (!Actor->OwnsComponent(SubObjectAsActorComponent))
+		if (!Actor->GetComponents().Contains(SubObjectAsActorComponent))
 		{
 			int32 MatchingComponents = 0;
 			UActorComponent* MappedActorComponent = SubObjectAsActorComponent;
@@ -389,7 +389,7 @@ void UProxyBackendNetDriver::InternalProcessRemoteFunction(
 				}
 			}
 
-			if (!Actor->OwnsComponent(MappedActorComponent))
+			if (!Actor->GetComponents().Contains(MappedActorComponent))
 			{
 				bShouldForwardRPC = false;
 
