Removing spectators from proxy for non-primary game server connections.

diff --git a/Engine/Plugins/Runtime/MultiServerReplication/Source/MultiServerReplication/Private/MultiServerProxy.cpp b/Engine/Plugins/Runtime/MultiServerReplication/Source/MultiServerReplication/Private/MultiServerProxy.cpp
--- a/Engine/Plugins/Runtime/MultiServerReplication/Source/MultiServerReplication/Private/MultiServerProxy.cpp
+++ b/Engine/Plugins/Runtime/MultiServerReplication/Source/MultiServerReplication/Private/MultiServerProxy.cpp
@@ -5,7 +5,6 @@
 #include "Engine/LocalPlayer.h"
 #include "Engine/GameInstance.h"
 #include "Engine/NetworkObjectList.h"
-#include "Engine/World.h"
 #include "GameFramework/GameModeBase.h"
 #include "Misc/CommandLine.h"
 #include "Net/DataChannel.h"
@@ -31,15 +30,20 @@ void UGameServerNotify::NotifyControlMessage(UNetConnection* GameServerConnectio
 {
 	Super::NotifyControlMessage(GameServerConnection, MessageType, Bunch);
 
-	// The default implementation of UPendingNetGame will only send the join request to the server when the level
-	// has loaded. Since the proxy is not currently dependent on loading levels we just send the join request 
-	// when receiving the welcome message to shortcut this logic.
 	if (MessageType == NMT_Welcome)
 	{
-		SendJoin();
+		// The default implementation of UPendingNetGame will only send the join request to the server when the level
+		// has loaded. Since the proxy is not currently dependent on loading levels we just send the join request 
+		// when receiving the welcome message to shortcut this logic.
+		SendJoinWithFlags(Flags);
 	}
 }
 
+void UGameServerNotify::SetProxyNetDriver(TObjectPtr<class UProxyNetDriver> InProxyNetDriver)
+{
+	ProxyNetDriver = InProxyNetDriver;
+}
+
 void UGameServerNotify::SetFirstPlayer(TObjectPtr<ULocalPlayer> Player)
 {
 	FirstPlayer = Player;
@@ -74,153 +78,161 @@ void UProxyListenerNotify::NotifyControlMessage(UNetConnection* ProxyConnection,
 {
 	check(ProxyNetDriver);
 
-	auto ProcessConnectionJoin = [this, ProxyConnection](int32 GameServerConnectionStateIndex, FGameServerConnectionState* GameServerConnectionState, bool bIsPrimaryGameServer)
+	// The NMT_Join message received by a proxy connection should trigger a connection to the backend game servers.
+	if (MessageType == NMT_Join)
 	{
-		// If this is the first connection to the game server, instantiate the backend network driver that will manage all
-		// connections from proxy connections to that server.
-		if (GameServerConnectionState->NetDriver == nullptr)
-		{
-			static int32 GameServerDriverIndex = 0;
-			FString Name = FString::Printf(TEXT("ProxyToGameServer-%d"), GameServerDriverIndex++);
-			GEngine->CreateNamedNetDriver(ProxyNetDriver->GetWorld(), *Name, "ProxyBackendNetDriver");
-
-			GameServerConnectionState->NetDriver = GEngine->FindNamedNetDriver(ProxyNetDriver->GetWorld(), *Name);
-			GameServerConnectionState->NetDriver->SetWorld(ProxyNetDriver->GetWorld());
-			Cast<UProxyBackendNetDriver>(GameServerConnectionState->NetDriver)->SetProxyNetDriver(ProxyNetDriver);
+		// For now, the primary game server is always the first registered game server.
+		const int32_t PrimaryGameServerIndex = 0;
 
-			UE_LOG(LogNetProxy, Log, TEXT("Created a game server NetDriver (name=%s, port=%d)"), *GameServerConnectionState->NetDriver->GetName(), GameServerConnectionState->Port);
+		for (int32 Index = 0; Index < ProxyNetDriver->GetGameServerConnectionCount(); Index++)
+		{
+			UE::Net::EJoinFlags Flags = UE::Net::EJoinFlags::NoPawn;
+			if (Index == PrimaryGameServerIndex)
+			{
+				EnumRemoveFlags(Flags, UE::Net::EJoinFlags::NoPawn);
+			}
+			ConnectToGameServer(ProxyConnection, PrimaryGameServerIndex, ProxyNetDriver->GetGameServerConnection(Index), Flags);
 		}
-		
-		UE_LOG(LogNetProxy, Log, TEXT("Client starting connection handshake on %s:%s"), *ProxyConnection->Driver->GetName(), *ProxyConnection->GetName());
+	}
+	else
+	{
+		// Forward all other connection messages onto the existing handshake logic.
+		ProxyNetDriver->GetWorld()->NotifyControlMessage(ProxyConnection, MessageType, Bunch);
+	}
+}
 
-		// Add a player to use on the game server.
-		const FPlatformUserId GameServerClientId = FPlatformUserId::CreateFromInternalId(ProxyNetDriver->GetNextGameServerClientId());
-		ULocalPlayer* NewPlayer = NewObject<ULocalPlayer>(GEngine, ULocalPlayer::StaticClass());
-		ProxyNetDriver->GetWorld()->GetGameInstance()->AddLocalPlayer(NewPlayer, GameServerClientId);
+void UProxyListenerNotify::ConnectToGameServer(UNetConnection* ProxyConnection, 
+											   int32 GameServerConnectionStateIndex,
+											   FGameServerConnectionState* GameServerConnectionState,
+											   UE::Net::EJoinFlags Flags)
+{
+	// If this is the first connection to the game server, instantiate the backend network driver that will manage all
+	// connections from proxy connections to that server.
+	if (GameServerConnectionState->NetDriver == nullptr)
+	{
+		static int32 GameServerDriverIndex = 0;
+		FString Name = FString::Printf(TEXT("ProxyToGameServer-%d"), GameServerDriverIndex++);
+		GEngine->CreateNamedNetDriver(ProxyNetDriver->GetWorld(), *Name, "ProxyBackendNetDriver");
 
-		// The new player will use the same unique identifier as the incoming proxy connection so that it will be propogated up to 
-		// the game servers through UNetConnection::PlayerId. This way each game server's incoming connection will have a PlayerId 
-		// that corresponds to a client connected to the proxy.
-		NewPlayer->SetCachedUniqueNetId(ProxyConnection->PlayerId);
+		GameServerConnectionState->NetDriver = GEngine->FindNamedNetDriver(ProxyNetDriver->GetWorld(), *Name);
+		GameServerConnectionState->NetDriver->SetWorld(ProxyNetDriver->GetWorld());
+		Cast<UProxyBackendNetDriver>(GameServerConnectionState->NetDriver)->SetProxyNetDriver(ProxyNetDriver);
 
-		UNetConnection* GameServerConnection = GameServerConnectionState->NetDriver->ServerConnection;
-		const bool bIsFirstGameServerConnection = (GameServerConnection == nullptr);
-		if (bIsFirstGameServerConnection)
-		{
-			const int32 GameServerPort = GameServerConnectionState->Port;
-			FString URLStr = FString::Format(TEXT("127.0.0.1:{0}"), { GameServerPort });
-			FURL URL(nullptr, *URLStr, TRAVEL_Absolute);
+		UE_LOG(LogNetProxy, Log, TEXT("Created a game server NetDriver (name=%s, port=%d)"), *GameServerConnectionState->NetDriver->GetName(), GameServerConnectionState->Port);
+	}
 
-			URL.AddOption(TEXT("AutonomousAsSimulated"));
+	// Add a player to use on the game server.
+	const FPlatformUserId GameServerClientId = FPlatformUserId::CreateFromInternalId(ProxyNetDriver->GetNextGameServerClientId());
+	ULocalPlayer* NewPlayer = NewObject<ULocalPlayer>(GEngine, ULocalPlayer::StaticClass());
+	ProxyNetDriver->GetWorld()->GetGameInstance()->AddLocalPlayer(NewPlayer, GameServerClientId);
 
-			if (!bIsPrimaryGameServer)
-			{
-				URL.AddOption(TEXT("SpectatorOnly=1"));
-			}
+	// The new player will use the same unique identifier as the incoming proxy connection so that it will be propogated up to 
+	// the game servers through UNetConnection::PlayerId. This way each game server's incoming connection will have a PlayerId 
+	// that corresponds to a client connected to the proxy.
+	NewPlayer->SetCachedUniqueNetId(ProxyConnection->PlayerId);
 
-			// Start the connection flow to the game server.
-			GameServerConnectionState->GameServerNotify = NewObject<UGameServerNotify>();
-			GameServerConnectionState->GameServerNotify->Initialize(URL);
-			GameServerConnectionState->GameServerNotify->InitNetDriver(GameServerConnectionState->NetDriver);
-			GameServerConnectionState->GameServerNotify->SetFirstPlayer(NewPlayer);
-			GameServerConnection = GameServerConnectionState->NetDriver->ServerConnection;
+	UNetConnection* GameServerConnection = GameServerConnectionState->NetDriver->ServerConnection;
+	const bool bIsFirstGameServerConnection = (GameServerConnection == nullptr);
+	if (bIsFirstGameServerConnection)
+	{
+		const int32 GameServerPort = GameServerConnectionState->Port;
+		FString URLStr = FString::Format(TEXT("127.0.0.1:{0}"), { GameServerPort });
+		FURL URL(nullptr, *URLStr, TRAVEL_Absolute);
+
+		// Maybe this should be a control message since it changes the server setting (it's global driver setting).
+		URL.AddOption(TEXT("AutonomousAsSimulated"));
+
+		// Start the connection flow to the game server.
+		GameServerConnectionState->GameServerNotify = NewObject<UGameServerNotify>();
+		GameServerConnectionState->GameServerNotify->Initialize(URL);
+		GameServerConnectionState->GameServerNotify->InitNetDriver(GameServerConnectionState->NetDriver);
+		GameServerConnectionState->GameServerNotify->SetFirstPlayer(NewPlayer);
+		GameServerConnectionState->GameServerNotify->SetProxyNetDriver(ProxyNetDriver);
+		GameServerConnectionState->GameServerNotify->SetFlags(Flags);
+		GameServerConnection = GameServerConnectionState->NetDriver->ServerConnection;
+
+		// UNetDriver::Notify will be reset in UPendingNetGame above so it's important that we override
+		// it here again to point to the proxy.
+		GameServerConnectionState->NetDriver->Notify = GameServerConnectionState->GameServerNotify;
+
+		UE_LOG(LogNetProxy, Log, TEXT("Connecting to game server: %s:%s -> %s (player=%s)"), 
+			   *ProxyConnection->GetDriver()->GetName(), 
+			   *ProxyConnection->GetName(),
+			   *GameServerConnection->GetDriver()->GetName(),
+			   *NewPlayer->GetName());
+	}
+	else
+	{
+		FGameServerSplitJoinRequest Request;
+		Request.Player = NewPlayer;
+		Request.Flags = Flags;
 
-			// UNetDriver::Notify will be reset in UPendingNetGame above so it's important that we override
-			// it here again to point to the proxy.
-			GameServerConnectionState->NetDriver->Notify = GameServerConnectionState->GameServerNotify;
+		// The NMT_JoinSplit message can only be sent when the parent connection is open.
+		if (GameServerConnectionState->NetDriver->ServerConnection->GetConnectionState() == USOCK_Open)
+		{
+			GameServerConnectionState->PendingSplitJoinRequests.Add(Request);
+			ProxyNetDriver->FlushSplitJoinRequests(GameServerConnectionState);
 
-			UE_LOG(LogNetProxy, Log, TEXT("Connecting to game server: %s:%s -> %s (player=%s)"), 
-				   *ProxyConnection->GetName(),
+			UE_LOG(LogNetProxy, Log, TEXT("Connecting (multiplexed) to game server: %s:%s -> %s (player=%s"), 
 				   *ProxyConnection->GetDriver()->GetName(), 
+				   *ProxyConnection->GetName(),
 				   *GameServerConnection->GetDriver()->GetName(),
 				   *NewPlayer->GetName());
 		}
 		else
 		{
-			FGameServerSplitJoinRequest Request;
-			Request.Player = NewPlayer;
-			Request.bIsPrimaryGameServer = bIsPrimaryGameServer;
-
-			// ULocalPlayer::SendSplitJoin() can only be called when the parent connection has been successfully opened. If 
-			// the parent connection is not open yet, this player will be queued up and processed when the parent connection
-			// has been opened (see UMultiProxyServer::GameServerConnectionHandshakeComplete()).
-			if (GameServerConnectionState->NetDriver->ServerConnection->GetConnectionState() == USOCK_Open)
-			{
-				GameServerConnectionState->PendingSplitJoinRequests.Add(Request);
-				ProxyNetDriver->FlushSplitJoinRequests(GameServerConnectionState);
-
-				UE_LOG(LogNetProxy, Log, TEXT("Connecting (multiplexed) to game server: %s:%s -> %s (player=%s"), 
-					   *ProxyConnection->GetName(),
-					   *ProxyConnection->GetDriver()->GetName(), 
-					   *GameServerConnection->GetDriver()->GetName(),
-					   *NewPlayer->GetName());
-			}
-			else
-			{
-				GameServerConnectionState->PendingSplitJoinRequests.Add(Request);
+			GameServerConnectionState->PendingSplitJoinRequests.Add(Request);
 	
-				UE_LOG(LogNetProxy, Log, TEXT("Queuing up connection (multiplexed) to game server: %s:%s -> %s (player=%s)"), 
-					   *ProxyConnection->GetName(),
-					   *ProxyConnection->GetDriver()->GetName(),
-					   *GameServerConnection->GetDriver()->GetName(),
-					   *NewPlayer->GetName());
-			}
+			UE_LOG(LogNetProxy, Log, TEXT("Queuing up connection (multiplexed) to game server: %s:%s -> %s (player=%s)"), 
+				   *ProxyConnection->GetDriver()->GetName(),
+				   *ProxyConnection->GetName(),
+				   *GameServerConnection->GetDriver()->GetName(),
+				   *NewPlayer->GetName());
 		}
+	}
 
-		const int32 NewPlayerIndex = GameServerConnectionState->Players.Num();
-		GameServerConnectionState->Players.Add(NewPlayer);
+	const int32 NewPlayerIndex = GameServerConnectionState->Players.Num();
+	GameServerConnectionState->Players.Add(NewPlayer);
 		
-		// Associate this proxy connection with the parent game server connection when beginning the handshake because
-		// the child connection hasn't been created yet. Once the handshake is complete, it's expected that this entry will
-		// be updated with the new child connection.
-		FGameServerAndPlayerIndex Key; 
-		Key.GameServerConnection = GameServerConnection;
-		Key.PlayerIndex = NewPlayerIndex;
-		ProxyNetDriver->GameToProxyConnection.Add(Key, ProxyConnection);
-
-		if (bIsPrimaryGameServer)
-		{
-			ProxyNetDriver->RegisterPrimaryGameServer(ProxyConnection, NewPlayer, GameServerConnectionStateIndex);
-		}
-	};
-
-	if (MessageType == NMT_Join)
-	{
-		// For now, the primary game server is always the first registered game server.
-		const int32_t PrimaryGameServerIndex = 0;
-
-		for (int32 Index = 0; Index < ProxyNetDriver->GetGameServerConnectionCount(); Index++)
-		{
-			const bool bIsPrimaryGameServer = (Index == PrimaryGameServerIndex);
-			ProcessConnectionJoin(PrimaryGameServerIndex, ProxyNetDriver->GetGameServerConnection(Index), bIsPrimaryGameServer);
-		}
-	}
-	else
+	// Associate this proxy connection with the parent game server connection when beginning the handshake because
+	// the child connection hasn't been created yet. Once the handshake is complete, it's expected that this entry will
+	// be updated with the new child connection.
+	FGameServerAndPlayerIndex Key; 
+	Key.GameServerConnection = GameServerConnection;
+	Key.PlayerIndex = NewPlayerIndex;
+	ProxyNetDriver->GameToProxyConnection.Add(Key, ProxyConnection);
+
+	if (!EnumHasAllFlags(Flags, UE::Net::EJoinFlags::NoPawn))
 	{
-		// Forward all other connection messages onto the existing handshake logic.
-		ProxyNetDriver->GetWorld()->NotifyControlMessage(ProxyConnection, MessageType, Bunch);
+		ProxyNetDriver->RegisterPrimaryGameServer(ProxyConnection, NewPlayer, GameServerConnectionStateIndex);
 	}
 }
 
 void UProxyBackendNetConnection::HandleClientPlayer(APlayerController* NewPlayerController, UNetConnection* GameServerConnection)
 {
+	// This function is called when a PlayerController is replicated to the proxy from a game server and represents the finalization
+	// of a connection to a primary or non-primary game server.
+
 	ensure(GameServerConnection == this);
 
 	UProxyBackendNetDriver* BackendNetDriver = Cast<UProxyBackendNetDriver>(Driver);
 	if (ensure(BackendNetDriver))
 	{
-		BackendNetDriver->GetProxyNetDriver()->GameServerConnectionHandshakeComplete(this, NewPlayerController, GameServerConnection);
+		BackendNetDriver->GetProxyNetDriver()->GameServerAssignPlayerController(this, GameServerConnection, NewPlayerController);
 	}
 }
 
 void UProxyBackendChildNetConnection::HandleClientPlayer(APlayerController* NewPlayerController, UNetConnection* GameServerConnection)
 {
+	// This function is called when a PlayerController is replicated to the proxy from a game server and represents the finalization
+	// of a connection to a primary or non-primary game server.
+
 	ensure(GameServerConnection != this);
 
 	UProxyBackendNetDriver* BackendNetDriver = Cast<UProxyBackendNetDriver>(Driver);
 	if (ensure(BackendNetDriver))
 	{
-		BackendNetDriver->GetProxyNetDriver()->GameServerConnectionHandshakeComplete(this, NewPlayerController, GameServerConnection);
+		BackendNetDriver->GetProxyNetDriver()->GameServerAssignPlayerController(this, GameServerConnection, NewPlayerController);
 	}
 }
 
@@ -414,11 +426,12 @@ void UProxyBackendNetDriver::InternalProcessRemoteFunction(
 
 	if (bShouldForwardRPC)
 	{
-		UE_LOG(LogNetProxy, Verbose, TEXT("RPC: Forwarding %s on %s (owner:%s sub-object:%s) to game server connection %s"), 
+		UE_LOG(LogNetProxy, Verbose, TEXT("RPC: Forwarding %s on %s (owner:%s sub-object:%s) to game server connection %s:%s"), 
 			   *Function->GetName(), 
 			   *Actor->GetName(), 
 			   *GetNameSafe(Actor->GetOwner()),
 			   *GetNameSafe(SubObject),
+			   *Connection->GetDriver()->GetName(),
 			   *Connection->GetName());
 
 		Super::InternalProcessRemoteFunction(Actor, SubObject, Connection, Function, Parms, OutParms, Stack, bIsServer);
@@ -697,17 +710,21 @@ bool UProxyNetDriver::IsProxySpawned(AActor* Actor) const
 	return ProxySpawnedActors.Contains(Actor);
 }
 
-void UProxyNetDriver::GameServerConnectionHandshakeComplete(UNetConnection* ChildGameServerConnection,
-															APlayerController* GameServerPlayerController,
-															UNetConnection* ParentGameServerConnection)
+void UProxyNetDriver::GameServerAssignPlayerController(UNetConnection* ChildGameServerConnection, UNetConnection* ParentGameServerConnection, APlayerController* GameServerPlayerController)
 {
-	ULocalPlayer* NewPlayer = GetPlayerOnGameServer(ParentGameServerConnection, GameServerPlayerController->NetPlayerIndex);
-	ensure(NewPlayer);
+	UE_LOG(LogNetProxy, Log, TEXT("Completing handshake for connection %s:%s (parent:%s:%s) to primary game server using game server controller %s:%s."), 
+		   *ChildGameServerConnection->GetDriver()->GetName(), *ChildGameServerConnection->GetName(), 
+		   *ParentGameServerConnection->GetDriver()->GetName(), *ParentGameServerConnection->GetName(),
+		   *GameServerPlayerController->GetName(), *GameServerPlayerController->GetClass()->GetName());
 
 	// If a new player controller is spawned on the proxy from a game server it's assumed it will replace any
 	// existing game server player controllers used by connections and players on the proxy. For this reason
 	// always detach an existing game server player controllers from the connection.
 	DetachPlayerController(ChildGameServerConnection->PlayerController);
+	
+	ULocalPlayer* NewPlayer = GetPlayerOnGameServer(ParentGameServerConnection, GameServerPlayerController->NetPlayerIndex);
+	ensure(NewPlayer);
+
 	AttachPlayerController(GameServerPlayerController, ChildGameServerConnection, NewPlayer);
 	
 	// The game server player controller is used as a client connection to the game server so must be ROLE_AutonomousProxy.
@@ -721,12 +738,14 @@ void UProxyNetDriver::GameServerConnectionHandshakeComplete(UNetConnection* Chil
 	const int32 PlayerIndex = GameServerPlayerController->NetPlayerIndex;
 	FGameServerAndPlayerIndex Key;
 	Key.GameServerConnection = ParentGameServerConnection;
- 	Key.PlayerIndex = PlayerIndex;
+	Key.PlayerIndex = PlayerIndex;
 
-	UE_LOG(LogNetProxy, VeryVerbose, TEXT("Lookup proxy connection associated with game server connection %s and player index %d (game_server_controller=%s"), 
-		   *ParentGameServerConnection->GetName(), 
-		   PlayerIndex,
-		   *GameServerPlayerController->GetName());
+	UE_LOG(LogNetProxy, VeryVerbose, TEXT("Lookup proxy connection associated with game server connection %s:%s and player index %d (game_server_controller=%s:%s"),
+		*ParentGameServerConnection->GetDriver()->GetName(),
+		*ParentGameServerConnection->GetName(),
+		PlayerIndex,
+		*GameServerPlayerController->GetName(),
+		*GameServerPlayerController->GetClass()->GetName());
 
 	// Find the proxy connection that established this connection to the game server.
 	TObjectPtr<UNetConnection>* ProxyConnectionPtr = GameToProxyConnection.Find(Key);
@@ -800,41 +819,44 @@ void UProxyNetDriver::GameServerConnectionHandshakeComplete(UNetConnection* Chil
 
 				ProxyPlayerToPrimaryGameServerPlayerController.Add(NewPlayer, GameServerPlayerController);
 
-				UE_LOG(LogNetProxy, Log, TEXT("Binding proxy connection %s:%s to primary game server connection %s:%s (game_server_controller=%s proxy_controller=%s player=%s:%d)"),
-					   *ProxyConnection->GetDriver()->GetName(),
-					   *ProxyConnection->GetName(),
-					   *ChildGameServerConnection->GetDriver()->GetName(),
-					   *ChildGameServerConnection->GetName(),
-					   *GameServerPlayerController->GetName(),
-					   *ProxyPlayerController->GetName(),
-					   *NewPlayer->GetName(),
-					   PlayerIndex);
+				UE_LOG(LogNetProxy, Log, TEXT("Binding proxy connection %s:%s to primary game server connection %s:%s (game_server_controller=%s:%s proxy_controller=%s:%s player=%s:%d)"),
+					*ProxyConnection->GetDriver()->GetName(),
+					*ProxyConnection->GetName(),
+					*ChildGameServerConnection->GetDriver()->GetName(),
+					*ChildGameServerConnection->GetName(),
+					*GameServerPlayerController->GetName(),
+					*GameServerPlayerController->GetClass()->GetName(),
+					*ProxyPlayerController->GetName(),
+					*ProxyPlayerController->GetClass()->GetName(),
+					*NewPlayer->GetName(),
+					PlayerIndex);
 			}
 			else
 			{
 				UE_LOG(LogNetProxy, Error, TEXT("Failed to bind a proxy connection %s:%s to a new primary game server %s:%s."),
-					   *ProxyConnection->GetDriver()->GetName(),
-					   *ProxyConnection->GetName(),
-					   *ChildGameServerConnection->GetDriver()->GetName(),
-					   *ChildGameServerConnection->GetName());
+					*ProxyConnection->GetDriver()->GetName(),
+					*ProxyConnection->GetName(),
+					*ChildGameServerConnection->GetDriver()->GetName(),
+					*ChildGameServerConnection->GetName());
 			}
 		}
 		else
 		{
-			UE_LOG(LogNetProxy, Log, TEXT("Ignoring connection %s:%s to non-primary game server %s:%s (game_server_controller=%s player=%s:%d)"),
+			UE_LOG(LogNetProxy, Log, TEXT("Ignoring connection %s:%s to non-primary game server %s:%s (game_server_controller=%s:%s player=%s:%d)"),
 				*ProxyConnection->GetDriver()->GetName(),
 				*ProxyConnection->GetName(),
 				*ChildGameServerConnection->GetDriver()->GetName(),
 				*ChildGameServerConnection->GetName(),
 				*GameServerPlayerController->GetName(),
+				*GameServerPlayerController->GetClass()->GetName(),
 				*NewPlayer->GetName(),
 				PlayerIndex);
 		}
 			
 		// Update the mapping of proxy connection to game server connection to use the new child game server connection (if it's 
 		// different from the parent connection).
-		FGameServerAndPlayerIndex NewKey; 
-		NewKey.GameServerConnection = ChildGameServerConnection; 
+		FGameServerAndPlayerIndex NewKey;
+		NewKey.GameServerConnection = ChildGameServerConnection;
 		NewKey.PlayerIndex = PlayerIndex;
 		GameToProxyConnection.Remove(Key);
 		GameToProxyConnection.Add(NewKey, ProxyConnection);
@@ -862,7 +884,8 @@ void UProxyNetDriver::GameServerConnectionHandshakeComplete(UNetConnection* Chil
 
 void UProxyNetDriver::FlushSplitJoinRequests(FGameServerConnectionState* GameServerConnectionState)
 {
-	if (GameServerConnectionState->NetDriver->ServerConnection->GetConnectionState() != EConnectionState::USOCK_Open)
+	EConnectionState ParentConnectionState = GameServerConnectionState->NetDriver->ServerConnection->GetConnectionState();
+	if (!(ParentConnectionState == EConnectionState::USOCK_Open))
 	{
 		UE_LOG(LogNetProxy, Error, TEXT("Flushing split join requests on %s without the parent connection being opened."), 
 			   *GameServerConnectionState->NetDriver->GetName());
@@ -870,20 +893,19 @@ void UProxyNetDriver::FlushSplitJoinRequests(FGameServerConnectionState* GameSer
 		return;
 	}
 
+	UE_LOG(LogNetProxy, Log, TEXT("Flushing %d split join connection requests for connection %s:%s."), 
+		   GameServerConnectionState->PendingSplitJoinRequests.Num(), *GameServerConnectionState->NetDriver->GetName(), *GameServerConnectionState->NetDriver->ServerConnection->GetName());
+
 	for (FGameServerSplitJoinRequest& Request : GameServerConnectionState->PendingSplitJoinRequests)
 	{
 		TArray<FString> Options;
 
-		if (!Request.bIsPrimaryGameServer)
-		{
-			Options.Add(TEXT("SpectatorOnly=1"));
-		}
-
-		Request.Player->SendSplitJoin(Options, GameServerConnectionState->NetDriver);
-
-		UE_LOG(LogNetProxy, Log, TEXT("Sending queued connection (multiplexed) request to game server: %s (player=%s"),
+		UE_LOG(LogNetProxy, Log, TEXT("Sending queued connection (multiplexed) request to game server: %s (player=%s flags=%d)"),
 			   *GameServerConnectionState->NetDriver->GetName(),
-			   *Request.Player->GetName());
+			   *Request.Player->GetName(),
+			   Request.Flags);
+
+		Request.Player->SendSplitJoin(Options, GameServerConnectionState->NetDriver, Request.Flags);
 	}
 
 	GameServerConnectionState->PendingSplitJoinRequests.Reset();
@@ -972,8 +994,9 @@ void UProxyNetDriver::DetachPlayerController(APlayerController* PlayerController
 	{
 		UPlayer* Player = PlayerController->Player;
 
-		UE_LOG(LogNetProxy, Log, TEXT("Destroying old controller %s on %s attached to player %s."),
+		UE_LOG(LogNetProxy, Log, TEXT("Destroying old controller %s:%s on %s attached to player %s."),
 			*PlayerController->GetName(),
+			*PlayerController->GetClass()->GetName(),
 			*this->GetName(),
 			*Player->GetName());
 
diff --git a/Engine/Plugins/Runtime/MultiServerReplication/Source/MultiServerReplication/Public/MultiServerProxy.h b/Engine/Plugins/Runtime/MultiServerReplication/Source/MultiServerReplication/Public/MultiServerProxy.h
--- a/Engine/Plugins/Runtime/MultiServerReplication/Source/MultiServerReplication/Public/MultiServerProxy.h
+++ b/Engine/Plugins/Runtime/MultiServerReplication/Source/MultiServerReplication/Public/MultiServerProxy.h
@@ -17,7 +17,7 @@ struct FGameServerSplitJoinRequest
 	UPROPERTY()
 	TObjectPtr<class ULocalPlayer> Player;
 
-	bool bIsPrimaryGameServer;
+	UE::Net::EJoinFlags Flags;
 };
 
 /**
@@ -69,8 +69,7 @@ struct FGameServerAndPlayerIndex
  *
  * Ensure that NMT_Join is sent after receiving NMT_Welcome from a game server. Normally
  * NMT_Join will be sent after a level is loaded but the proxy doesn't currently handle 
- * loading levels when connecting to a server. IMPORTANT: Once level loading functionality 
- * has been added, this class can be removed.
+ * loading levels when connecting to a server.
  *
  * Defaults to the behavior in UPendingNetGame which normally handles all outgoing 
  * connections to a game server.
@@ -88,12 +87,21 @@ public:
 	virtual void 					NotifyControlMessage(UNetConnection* GameServerConnection, uint8 MessageType, class FInBunch& Bunch) override;
 	virtual class ULocalPlayer* 	GetFirstGamePlayer();
 
+	void SetProxyNetDriver(TObjectPtr<class UProxyNetDriver> InProxyNetDriver);
+	void SetFlags(UE::Net::EJoinFlags InFlags) { Flags = InFlags; };
+
 	/** Set the value to be returned by GetFirstGamePlayer(). */
 	void SetFirstPlayer(TObjectPtr<class ULocalPlayer> Player);
 
 private:
+
+	UE::Net::EJoinFlags Flags;
+
 	UPROPERTY()
 	TObjectPtr<class ULocalPlayer> FirstPlayer;
+	
+	UPROPERTY()
+	TObjectPtr<class UProxyNetDriver> ProxyNetDriver;
 };
 
 /** 
@@ -123,6 +131,12 @@ public:
 
 private:
 
+	/** Start connecting an incoming proxy connection to a game server. */
+	void ConnectToGameServer(UNetConnection* ProxyConnection, 
+							 int32 GameServerConnectionStateIndex, 
+							 FGameServerConnectionState* GameServerConnectionState, 
+							 UE::Net::EJoinFlags Flags);
+
 	UPROPERTY()
 	TObjectPtr<class UProxyNetDriver> ProxyNetDriver;
 };
@@ -219,6 +233,15 @@ private:
 * IMPORTANT: There isn't a relationship from LocalPlayer to GameServerPlayerController because
 * the exiting pointer (ULocalPlayer::PlayerController) that is being used can only point to one
 * controller.
+*
+* For each proxy connection one of the game servers are considered the primary game server. This
+* is the game server that spawns the proxy client's pawn, player controller, receive RPCs from
+* the proxy connection, and send RPCs to the proxy connection. The other game servers are considered
+* non-primary game servers and only replicate state relevant to that connection to the proxy.
+*
+* When connecting to non-primary game servers the game server will spawn a ANoPawnPlayerController 
+* player controller, and not spawn a pawn. These connections will replicate state from the game 
+* server but not maintain a player presence.
 */
 UCLASS()
 class MULTISERVERREPLICATION_API UProxyNetDriver : public UIpNetDriver
@@ -255,10 +278,8 @@ public:
 	/** Get the next identifier for outgoing connections to game servers. */
 	int32 GetNextGameServerClientId();
 	
-	/** Called when a parent or child game server connection has been established. */
-	void GameServerConnectionHandshakeComplete(UNetConnection* ChildGameServerConnection,
-											   APlayerController* PlayerController,
-											   UNetConnection* ParentGameServerConnection);
+	/** Called when the player controller associated with a connection is changed (either at the end of initial connection handshake, or if changed after successfully connected). */
+	void GameServerAssignPlayerController(UNetConnection* ChildGameServerConnection, UNetConnection* ParentGameServerConnection, APlayerController* PlayerController);
 
 private:
 
diff --git a/Engine/Source/Runtime/Core/Public/Misc/EngineNetworkCustomVersion.h b/Engine/Source/Runtime/Core/Public/Misc/EngineNetworkCustomVersion.h
--- a/Engine/Source/Runtime/Core/Public/Misc/EngineNetworkCustomVersion.h
+++ b/Engine/Source/Runtime/Core/Public/Misc/EngineNetworkCustomVersion.h
@@ -92,6 +92,7 @@ struct FEngineNetworkCustomVersion
 		MontagePlayCountSerialization = 37,		// Bump version to support serialization changes for dynamic montage play count 
 		RemoteObjectReferences = 38,			// Bump version to support serialization changes for remote object references 
 		BeaconNetIDVariant = 39,				// Bump version for online beacon's BeaconAssignGUID control message switching to use FNetIDVariant
+		JoinNoPawn = 40,						// Bump version to support NMT_JoinNoPawn and NMT_JoinNoPawnSplit control messages.
 
 		// -----<new versions can be added above this line>-------------------------------------------------
 		VersionPlusOne,
diff --git a/Engine/Source/Runtime/Engine/Classes/Engine/LocalPlayer.h b/Engine/Source/Runtime/Engine/Classes/Engine/LocalPlayer.h
--- a/Engine/Source/Runtime/Engine/Classes/Engine/LocalPlayer.h
+++ b/Engine/Source/Runtime/Engine/Classes/Engine/LocalPlayer.h
@@ -462,9 +462,12 @@ public:
 	ENGINE_API virtual void SendSplitJoin(TArray<FString>& Options);
 
 	/** 
-	 * Same as SendSplitJoin() above but with a additional NetDriver argument which is the drvier to send the request.
+	 * Same as SendSplitJoin() above but with additional arguments.
+	 *
+	 * @param	NetDriver		Use this NetDriver to send the join request instead of the NetDriver associated with the world.
+	 * @param	Flags			Configuration flags used by the join request sent to the server.
 	 */
-	ENGINE_API virtual void SendSplitJoin(TArray<FString>& Options, UNetDriver* NetDriver);
+	ENGINE_API virtual void SendSplitJoin(TArray<FString>& Options, UNetDriver* NetDriver, UE::Net::EJoinFlags Flags);
 	
 	/**
 	 * Change the physical ControllerId for this player; if the specified ControllerId is already taken by another player, changes the ControllerId
diff --git a/Engine/Source/Runtime/Engine/Classes/Engine/NetDriver.h b/Engine/Source/Runtime/Engine/Classes/Engine/NetDriver.h
--- a/Engine/Source/Runtime/Engine/Classes/Engine/NetDriver.h
+++ b/Engine/Source/Runtime/Engine/Classes/Engine/NetDriver.h
@@ -2336,6 +2336,15 @@ public:
 		bUseRemoteObjectRefs = bValue;
 	}
 
+	/** Enable/disable support for NMT_JoinNoPawn and NMT_JoinNoPawnSplit. */
+	void SetSupportForNoPawnConnection(bool bInSupportNoPawnConnection)
+	{
+		bSupportNoPawnConnection = bInSupportNoPawnConnection;
+	}
+
+	/** Return true if the driver supports NMT_JoinNoPawn and NMT_JoinNoPawnSplit. */
+	bool HasSupportForNoPawnConnection();
+
 protected:
 	
 	/** Stream of random numbers to be used by this instance of UNetDriver */
@@ -2480,4 +2489,7 @@ private:
 
 	/** When using remote objects with UE_WITH_REMOTE_OBJECT_HANDLE, determines whether or not references to UObjects (either from TObjectPtr or TWeakObjectPtr) are serialized as FRemoteObjectReference */
 	bool bUseRemoteObjectRefs = false;
+
+	/** Server support for NMT_JoinNoPawn and NMT_JoinNoPawnSplit. */
+	bool bSupportNoPawnConnection = false;
 };
diff --git a/Engine/Source/Runtime/Engine/Classes/Engine/PendingNetGame.h b/Engine/Source/Runtime/Engine/Classes/Engine/PendingNetGame.h
--- a/Engine/Source/Runtime/Engine/Classes/Engine/PendingNetGame.h
+++ b/Engine/Source/Runtime/Engine/Classes/Engine/PendingNetGame.h
@@ -129,6 +129,9 @@ public:
 	/** Send JOIN to other end */
 	virtual void SendJoin();
 
+	/** Send Join to the server with optional flags */
+	virtual void SendJoinWithFlags(UE::Net::EJoinFlags Flags);
+
 	//~ Begin UObject Interface.
 	virtual void Serialize( FArchive& Ar ) override;
 
diff --git a/Engine/Source/Runtime/Engine/Classes/GameFramework/PlayerController.h b/Engine/Source/Runtime/Engine/Classes/GameFramework/PlayerController.h
--- a/Engine/Source/Runtime/Engine/Classes/GameFramework/PlayerController.h
+++ b/Engine/Source/Runtime/Engine/Classes/GameFramework/PlayerController.h
@@ -2273,3 +2273,14 @@ public:
 	bool GetNetworkPhysicsTickOffsetAssigned() const { return bNetworkPhysicsTickOffsetAssigned; }
 
 };
+
+/** 
+ * A specialized player controller used by server connections that are created using NMT_JoinNoPawn and NMT_JoinNoPawnSplit.
+ */
+UCLASS(MinimalAPI)
+class ANoPawnPlayerController : public APlayerController
+{
+public:
+
+	GENERATED_BODY()
+};
diff --git a/Engine/Source/Runtime/Engine/Private/DataChannel.cpp b/Engine/Source/Runtime/Engine/Private/DataChannel.cpp
--- a/Engine/Source/Runtime/Engine/Private/DataChannel.cpp
+++ b/Engine/Source/Runtime/Engine/Private/DataChannel.cpp
@@ -2079,6 +2079,12 @@ void UControlChannel::ReceivedBunch( FInBunch& Bunch )
 				case NMT_IrisNetRefHandleError:
 					FNetControlMessage<NMT_IrisNetRefHandleError>::Discard(Bunch);
 					break;
+				case NMT_JoinNoPawn:
+					FNetControlMessage<NMT_JoinNoPawn>::Discard(Bunch);
+					break;
+				case NMT_JoinNoPawnSplit:
+					FNetControlMessage<NMT_JoinNoPawnSplit>::Discard(Bunch);
+					break;
 				default:
 					// if this fails, a case is missing above for an implemented message type
 					// or the connection is being sent potentially malformed packets
diff --git a/Engine/Source/Runtime/Engine/Private/LocalPlayer.cpp b/Engine/Source/Runtime/Engine/Private/LocalPlayer.cpp
--- a/Engine/Source/Runtime/Engine/Private/LocalPlayer.cpp
+++ b/Engine/Source/Runtime/Engine/Private/LocalPlayer.cpp
@@ -322,10 +322,10 @@ void ULocalPlayer::SendSplitJoin(TArray<FString>& Options)
 		NetDriver = World->GetNetDriver();
 	}
 
-	SendSplitJoin(Options, NetDriver);
+	SendSplitJoin(Options, NetDriver, UE::Net::EJoinFlags::None);
 }
 
-void ULocalPlayer::SendSplitJoin(TArray<FString>& Options, UNetDriver* NetDriver)
+void ULocalPlayer::SendSplitJoin(TArray<FString>& Options, UNetDriver* NetDriver, UE::Net::EJoinFlags Flags)
 {
 	if (NetDriver == NULL || NetDriver->ServerConnection == NULL || NetDriver->ServerConnection->GetConnectionState() != USOCK_Open)
 	{
@@ -381,7 +381,16 @@ void ULocalPlayer::SendSplitJoin(TArray<FString>& Options, UNetDriver* NetDriver
 			FUniqueNetIdRepl UniqueIdRepl(GetPreferredUniqueNetId());
 
 			FString URLString = URL.ToString();
-			FNetControlMessage<NMT_JoinSplit>::Send(NetDriver->ServerConnection, URLString, UniqueIdRepl);
+
+			if (EnumHasAllFlags(Flags, UE::Net::EJoinFlags::NoPawn))
+			{
+				FNetControlMessage<NMT_JoinNoPawnSplit>::Send(NetDriver->ServerConnection, UniqueIdRepl);
+			}
+			else
+			{
+				FNetControlMessage<NMT_JoinSplit>::Send(NetDriver->ServerConnection, URLString, UniqueIdRepl);
+			}
+			
 			bSentSplitJoin = true;
 		}
 	}
diff --git a/Engine/Source/Runtime/Engine/Private/NetDriver.cpp b/Engine/Source/Runtime/Engine/Private/NetDriver.cpp
--- a/Engine/Source/Runtime/Engine/Private/NetDriver.cpp
+++ b/Engine/Source/Runtime/Engine/Private/NetDriver.cpp
@@ -466,6 +466,13 @@ static FAutoConsoleVariableRef CVarEnableNetStats(
 	TEXT("Enable the calculation of network stats and storing them in the network metrics database."),
 	ECVF_Default);
 
+static bool GbEnableSupportForNoPawnConnection = false;
+static FAutoConsoleVariableRef CVarEnableSupportForNoPawnConnection(
+	TEXT("net.EnableSupportForNoPawnConnection"),
+	GbEnableSupportForNoPawnConnection,
+	TEXT("Enable server support for NMT_JoinNoPawn and NMT_JoinNoPawnSplit control messages."),
+	ECVF_Default);
+
 namespace UE::Net
 {
 	static FString GRequiredEncryptionNetDriverDefNames_Internal = TEXT("all");
@@ -789,6 +796,11 @@ void UNetDriver::NotifyGameInstanceUpdated()
 #endif
 }
 
+bool UNetDriver::HasSupportForNoPawnConnection()
+{
+	return GbEnableSupportForNoPawnConnection || bSupportNoPawnConnection;
+}
+
 void UNetDriver::AssertValid()
 {
 }
diff --git a/Engine/Source/Runtime/Engine/Private/PendingNetGame.cpp b/Engine/Source/Runtime/Engine/Private/PendingNetGame.cpp
--- a/Engine/Source/Runtime/Engine/Private/PendingNetGame.cpp
+++ b/Engine/Source/Runtime/Engine/Private/PendingNetGame.cpp
@@ -559,11 +559,24 @@ void UPendingNetGame::Tick( float DeltaTime )
 }
 
 void UPendingNetGame::SendJoin()
+{
+	SendJoinWithFlags(UE::Net::EJoinFlags::None);
+}
+
+void UPendingNetGame::SendJoinWithFlags(UE::Net::EJoinFlags Flags)
 {
 	bSentJoinRequest = true;
 
-	FNetControlMessage<NMT_Join>::Send(NetDriver->ServerConnection);
-	NetDriver->ServerConnection->FlushNet(true);
+	if (EnumHasAllFlags(Flags, UE::Net::EJoinFlags::NoPawn))
+	{
+		FNetControlMessage<NMT_JoinNoPawn>::Send(NetDriver->ServerConnection);
+		NetDriver->ServerConnection->FlushNet(true);
+	}
+	else
+	{
+		FNetControlMessage<NMT_Join>::Send(NetDriver->ServerConnection);
+		NetDriver->ServerConnection->FlushNet(true);
+	}
 }
 
 ULocalPlayer* UPendingNetGame::GetFirstGamePlayer()
diff --git a/Engine/Source/Runtime/Engine/Private/World.cpp b/Engine/Source/Runtime/Engine/Private/World.cpp
--- a/Engine/Source/Runtime/Engine/Private/World.cpp
+++ b/Engine/Source/Runtime/Engine/Private/World.cpp
@@ -7105,6 +7105,73 @@ void UWorld::NotifyControlMessage(UNetConnection* Connection, uint8 MessageType,
 
 				break;
 			}
+			case NMT_JoinNoPawn:
+			{
+				if (NetDriver->HasSupportForNoPawnConnection())
+				{
+					FActorSpawnParameters SpawnInfo;
+					SpawnInfo.ObjectFlags |= RF_Transient;
+					SpawnInfo.bDeferConstruction = true;
+					
+					Connection->PlayerController = GetWorld()->SpawnActor<APlayerController>(ANoPawnPlayerController::StaticClass(), FVector(), FRotator(), SpawnInfo);
+					Connection->PlayerController->SetRole(ROLE_Authority);
+					Connection->PlayerController->SetReplicates(true);
+					Connection->PlayerController->SetAutonomousProxy(true);
+					Connection->PlayerController->FinishSpawning(FTransform());
+					Connection->PlayerController->NetPlayerIndex = 0;
+					Connection->PlayerController->NetConnection = Connection;
+
+					Connection->OwningActor = Connection->PlayerController;
+				}
+				else
+				{
+					UE_LOG(LogNet, Log, TEXT("Server doesn't support NMT_JoinNoPawn (enable with UNetDriver::SetSupportForNoPawnConnection() or net.EnableSupportForNoPawnConnection)."));
+					FString Msg(TEXT("Server doesn't support NMT_JoinNoPawn."));
+					Connection->SendCloseReason(ENetCloseResult::JoinFailure);
+					FNetControlMessage<NMT_Failure>::Send(Connection, Msg);
+					Connection->FlushNet(true);
+					Connection->Close(ENetCloseResult::JoinFailure);
+				}
+
+				break;
+			}
+			case NMT_JoinNoPawnSplit:
+			{
+				if (NetDriver->HasSupportForNoPawnConnection())
+				{
+					FUniqueNetIdRepl UniqueIdRepl;
+					if (FNetControlMessage<NMT_JoinNoPawnSplit>::Receive(Bunch, UniqueIdRepl))
+					{
+						FActorSpawnParameters SpawnInfo;
+						SpawnInfo.ObjectFlags |= RF_Transient;
+						SpawnInfo.bDeferConstruction = true;
+					
+						UChildConnection* ChildConnection = NetDriver->CreateChild(Connection);
+						ChildConnection->PlayerId = UniqueIdRepl;
+
+						ChildConnection->PlayerController = GetWorld()->SpawnActor<APlayerController>(ANoPawnPlayerController::StaticClass(), FVector(), FRotator(), SpawnInfo);
+						ChildConnection->PlayerController->SetRole(ROLE_Authority);
+						ChildConnection->PlayerController->SetReplicates(true);
+						ChildConnection->PlayerController->SetAutonomousProxy(true);
+						ChildConnection->PlayerController->FinishSpawning(FTransform());
+						ChildConnection->PlayerController->NetPlayerIndex = Connection->Children.Num();
+						ChildConnection->PlayerController->NetConnection = ChildConnection;
+
+						ChildConnection->OwningActor = ChildConnection->PlayerController;
+					}
+				}
+				else
+				{
+					UE_LOG(LogNet, Log, TEXT("Server doesn't support NMT_JoinNoPawnSplit (enable with UNetDriver::SetSupportForNoPawnConnection() or net.EnableSupportForNoPawnConnection)."));
+					FString Msg(TEXT("Server doesn't support NMT_JoinNoPawnSplit."));
+					Connection->SendCloseReason(ENetCloseResult::JoinFailure);
+					FNetControlMessage<NMT_Failure>::Send(Connection, Msg);
+					Connection->FlushNet(true);
+					Connection->Close(ENetCloseResult::JoinFailure);
+				}
+
+				break;
+			}
 			case NMT_PCSwap:
 			{
 				UNetConnection* SwapConnection = Connection;
diff --git a/Engine/Source/Runtime/Engine/Public/Net/DataChannel.h b/Engine/Source/Runtime/Engine/Public/Net/DataChannel.h
--- a/Engine/Source/Runtime/Engine/Public/Net/DataChannel.h
+++ b/Engine/Source/Runtime/Engine/Public/Net/DataChannel.h
@@ -217,3 +217,15 @@ DEFINE_CONTROL_CHANNEL_MESSAGE(BeaconNetGUIDAck, 28, FString); // server receive
 
 DEFINE_CONTROL_CHANNEL_MESSAGE(IrisProtocolMismatch, 29, uint64); // client has a different protocol hash from the server
 DEFINE_CONTROL_CHANNEL_MESSAGE(IrisNetRefHandleError, 30, UE::Net::ENetRefHandleError, uint64); // a specific handle caused an error and we want the remote connection to log all information it has on it
+
+// Called in place of Join or JoinSplit during connection handshake, and will created a connection that has the following properties:
+//	* Won't spawn an instance of APawn.
+//	* Won't call AGameModeBase::Login().
+//	* Spawns ANoPawnPlayerController instead of using the class in AGameMode::PlayerControllerClass.
+//	* Network relevency will be based on the location of ANoPawnPlayerController.
+//
+// These messages are only intended to be used when called from one game server to another game server. For this reason it is disabled
+// by default and must be explicitly enabled on the receiving server using the cvar net.EnableSupportForNoPawnConnection or calling
+// UNetDriver::SetSupportForNoPawnConnection() directly.
+DEFINE_CONTROL_CHANNEL_MESSAGE(JoinNoPawn, 31);
+DEFINE_CONTROL_CHANNEL_MESSAGE(JoinNoPawnSplit, 32, FUniqueNetIdRepl);
\ No newline at end of file
diff --git a/Engine/Source/Runtime/Net/Core/Public/Net/Core/Connection/NetEnums.h b/Engine/Source/Runtime/Net/Core/Public/Net/Core/Connection/NetEnums.h
--- a/Engine/Source/Runtime/Net/Core/Public/Net/Core/Connection/NetEnums.h
+++ b/Engine/Source/Runtime/Net/Core/Public/Net/Core/Connection/NetEnums.h
@@ -6,6 +6,16 @@
 
 #include "NetEnums.generated.h"
 
+namespace UE::Net
+{
+	enum class EJoinFlags
+	{
+		None,
+		NoPawn = 1,
+	};
+	ENUM_CLASS_FLAGS(EJoinFlags)
+}
+
 /** Types of network failures broadcast from the engine */
 UENUM(BlueprintType)
 namespace ENetworkFailure
