Fix to component RPCs not always routing correctly through the proxy.

diff --git a/Engine/Plugins/Runtime/MultiServerReplication/Source/MultiServerReplication/Private/MultiServerProxy.cpp b/Engine/Plugins/Runtime/MultiServerReplication/Source/MultiServerReplication/Private/MultiServerProxy.cpp
--- a/Engine/Plugins/Runtime/MultiServerReplication/Source/MultiServerReplication/Private/MultiServerProxy.cpp
+++ b/Engine/Plugins/Runtime/MultiServerReplication/Source/MultiServerReplication/Private/MultiServerProxy.cpp
@@ -382,7 +382,7 @@ void UProxyBackendNetDriver::InternalProcessRemoteFunction(
 	// This logic assumes that an actor only has one component of a given type and will fail if that assumption is incorrect.
 	if (UActorComponent* SubObjectAsActorComponent = Cast<UActorComponent>(SubObject))
 	{
-		if (!Actor->GetComponents().Contains(SubObjectAsActorComponent))
+		if (SubObjectAsActorComponent->GetOwner() != Actor)
 		{
 			int32 MatchingComponents = 0;
 			UActorComponent* MappedActorComponent = SubObjectAsActorComponent;
@@ -401,7 +401,7 @@ void UProxyBackendNetDriver::InternalProcessRemoteFunction(
 				}
 			}
 
-			if (!Actor->GetComponents().Contains(MappedActorComponent))
+			if (MappedActorComponent->GetOwner() != Actor)
 			{
 				bShouldForwardRPC = false;
 
