Proxy command line arguments to specify and cycle the primary game server used when clients connect.

diff --git a/Engine/Plugins/Runtime/MultiServerReplication/Source/MultiServerReplication/Private/MultiServerProxy.cpp b/Engine/Plugins/Runtime/MultiServerReplication/Source/MultiServerReplication/Private/MultiServerProxy.cpp
--- a/Engine/Plugins/Runtime/MultiServerReplication/Source/MultiServerReplication/Private/MultiServerProxy.cpp
+++ b/Engine/Plugins/Runtime/MultiServerReplication/Source/MultiServerReplication/Private/MultiServerProxy.cpp
@@ -82,16 +82,19 @@ void UProxyListenerNotify::NotifyControlMessage(UNetConnection* ProxyConnection,
 	if (MessageType == NMT_Join)
 	{
 		// For now, the primary game server is always the first registered game server.
-		const int32_t PrimaryGameServerIndex = 0;
-
 		for (int32 Index = 0; Index < ProxyNetDriver->GetGameServerConnectionCount(); Index++)
 		{
 			UE::Net::EJoinFlags Flags = UE::Net::EJoinFlags::NoPawn;
-			if (Index == PrimaryGameServerIndex)
+			if (Index == ProxyNetDriver->PrimaryGameServerForNextClient)
 			{
 				EnumRemoveFlags(Flags, UE::Net::EJoinFlags::NoPawn);
 			}
-			ConnectToGameServer(ProxyConnection, PrimaryGameServerIndex, ProxyNetDriver->GetGameServerConnection(Index), Flags);
+			ConnectToGameServer(ProxyConnection, ProxyNetDriver->PrimaryGameServerForNextClient, ProxyNetDriver->GetGameServerConnection(Index), Flags);
+		}
+
+		if (ProxyNetDriver->bCyclePrimaryGameServer)
+		{
+			ProxyNetDriver->PrimaryGameServerForNextClient = (ProxyNetDriver->PrimaryGameServerForNextClient + 1) % ProxyNetDriver->GetGameServerConnectionCount();
 		}
 	}
 	else
@@ -518,6 +521,17 @@ bool UProxyNetDriver::InitBase(bool bInitAsClient, FNetworkNotify* InNotify, con
 		}
 	}
 
+	FString ClientPrimaryGameServer;
+	if (FParse::Value(FCommandLine::Get(), TEXT("ProxyClientPrimaryGameServer="), ClientPrimaryGameServer))
+	{
+		PrimaryGameServerForNextClient = FCString::Atoi(ToCStr(ClientPrimaryGameServer));
+	}
+
+	if (FParse::Param(FCommandLine::Get(), TEXT("ProxyCyclePrimaryGameServer")))
+	{
+		bCyclePrimaryGameServer = true;
+	}
+
 	SetReplicateTransactionally(false);
 	
 #if UE_SUPPORT_FOR_ACTOR_TICK_DISABLE
diff --git a/Engine/Plugins/Runtime/MultiServerReplication/Source/MultiServerReplication/Public/MultiServerProxy.h b/Engine/Plugins/Runtime/MultiServerReplication/Source/MultiServerReplication/Public/MultiServerProxy.h
--- a/Engine/Plugins/Runtime/MultiServerReplication/Source/MultiServerReplication/Public/MultiServerProxy.h
+++ b/Engine/Plugins/Runtime/MultiServerReplication/Source/MultiServerReplication/Public/MultiServerProxy.h
@@ -354,4 +354,10 @@ private:
 	TSet<TObjectPtr<AActor>> ProxySpawnedActors;
 
 	int32 NextGameServerClientId = 0;
+	
+	// The primary game server to use for the next primary game client.
+	int32 PrimaryGameServerForNextClient = 0;
+
+	// After a client connects to the proxy increment the value of PrimaryGameServerForNextClient.
+	bool bCyclePrimaryGameServer = false;
 };
diff --git a/Engine/Source/Programs/AutomationTool/Scripts/LaunchMultiServer.Automation.cs b/Engine/Source/Programs/AutomationTool/Scripts/LaunchMultiServer.Automation.cs
--- a/Engine/Source/Programs/AutomationTool/Scripts/LaunchMultiServer.Automation.cs
+++ b/Engine/Source/Programs/AutomationTool/Scripts/LaunchMultiServer.Automation.cs
@@ -470,11 +470,30 @@ namespace AutomationTool
 							Args += "-WaitForDebuggerNoBreak ";
 						}
 
-						Args += "-ProxyGameServers=";
+						if (ServerInfos.Values.Count() > 0)
+						{
+							Args += "-ProxyGameServers=";
+
+							foreach (var ServerInfo in ServerInfos.Values)
+							{
+								Args += String.Format("127.0.0.1:{0},", ServerInfo.GameListenPort);
+							}
+
+							Args += " ";
+						}
+
+						if (ParseParam("proxyclientprimary"))
+						{
+							int ProxyClientPrimary = 0;
+							string ProxyClientPrimaryStr = ParseParamValue("proxyclientprimary", null);
+							Int32.TryParse(ProxyClientPrimaryStr, out ProxyClientPrimary);
+
+							Args += String.Format("-ProxyClientPrimaryGameServer={0} ", ProxyClientPrimary);
+						}
 
-						foreach (var ServerInfo in ServerInfos.Values)
+						if (ParseParam("proxycycleprimary"))
 						{
-							Args += String.Format("127.0.0.1:{0},", ServerInfo.GameListenPort);
+							Args += "-ProxyCyclePrimaryGameServer ";
 						}
 					}
 
