Support for player migration in MultiServer Proxy.

diff --git a/Engine/Plugins/Runtime/GameplayStateTree/Source/GameplayStateTreeModule/Private/Components/StateTreeComponentSchema.cpp b/Engine/Plugins/Runtime/GameplayStateTree/Source/GameplayStateTreeModule/Private/Components/StateTreeComponentSchema.cpp
--- a/Engine/Plugins/Runtime/GameplayStateTree/Source/GameplayStateTreeModule/Private/Components/StateTreeComponentSchema.cpp
+++ b/Engine/Plugins/Runtime/GameplayStateTree/Source/GameplayStateTreeModule/Private/Components/StateTreeComponentSchema.cpp
@@ -217,7 +217,7 @@ bool UStateTreeComponentSchema::CollectExternalData(const FStateTreeExecutionCon
 			}
 			else if (ItemDesc.Struct->IsChildOf(APawn::StaticClass()))
 			{
-				APawn* OwnerPawn = (AIOwner != nullptr) ? AIOwner->GetPawn() : Cast<APawn>(Owner);
+				APawn* OwnerPawn = (AIOwner != nullptr) ? AIOwner->GetPawn().Get() : Cast<APawn>(Owner);
 				UE_CVLOG(OwnerPawn == nullptr, Context.GetOwner(), LogStateTree, Error, TEXT("StateTree %s: Could not find required pawn %s"), *GetNameSafe(Context.GetStateTree()), *GetNameSafe(ItemDesc.Struct));
 				OutDataViews[Index] = FStateTreeDataView(OwnerPawn);
 				IssuesFoundCounter += OwnerPawn != nullptr ? 0 : 1;
diff --git a/Engine/Plugins/Runtime/MultiServerReplication/Source/MultiServerReplication/Private/MultiServerProxy.cpp b/Engine/Plugins/Runtime/MultiServerReplication/Source/MultiServerReplication/Private/MultiServerProxy.cpp
--- a/Engine/Plugins/Runtime/MultiServerReplication/Source/MultiServerReplication/Private/MultiServerProxy.cpp
+++ b/Engine/Plugins/Runtime/MultiServerReplication/Source/MultiServerReplication/Private/MultiServerProxy.cpp
@@ -11,6 +11,21 @@
 
 DEFINE_LOG_CATEGORY_STATIC(LogNetProxy, Log, All);
 
+FString LexToString(EProxyConnectionState State)
+{
+	switch (State)
+	{
+		case EProxyConnectionState::Disconnected:
+			return FString(TEXT("Disconnected"));
+		case EProxyConnectionState::ConnectingPrimary:
+			return FString(TEXT("ConnectingPrimary"));
+		case EProxyConnectionState::ConnectedPrimary:
+			return FString(TEXT("ConnectedPrimary"));
+		default:
+			return FString(TEXT("Unknown"));
+	}
+}
+
 void UGameServerNotify::NotifyAcceptedConnection(class UNetConnection* Connection)
 {
 
@@ -109,6 +124,8 @@ void UProxyListenerNotify::ConnectToGameServer(UNetConnection* ProxyConnection,
 											   FGameServerConnectionState* GameServerConnectionState,
 											   UE::Net::EJoinFlags Flags)
 {
+	ProxyNetDriver->ProxyConnectionState.Add(ProxyConnection, EProxyConnectionState::Disconnected);
+
 	// If this is the first connection to the game server, instantiate the backend network driver that will manage all
 	// connections from proxy connections to that server.
 	if (GameServerConnectionState->NetDriver == nullptr)
@@ -134,6 +151,8 @@ void UProxyListenerNotify::ConnectToGameServer(UNetConnection* ProxyConnection,
 	// that corresponds to a client connected to the proxy.
 	NewPlayer->SetCachedUniqueNetId(ProxyConnection->PlayerId);
 
+	const uint32 ClientHandshakeId = ProxyNetDriver->GetNextClientHandshakeId();
+
 	UNetConnection* GameServerConnection = GameServerConnectionState->NetDriver->ServerConnection;
 	const bool bIsFirstGameServerConnection = (GameServerConnection == nullptr);
 	if (bIsFirstGameServerConnection)
@@ -141,6 +160,8 @@ void UProxyListenerNotify::ConnectToGameServer(UNetConnection* ProxyConnection,
 		const int32 GameServerPort = GameServerConnectionState->Port;
 		FString URLStr = FString::Format(TEXT("127.0.0.1:{0}"), { GameServerPort });
 		FURL URL(nullptr, *URLStr, TRAVEL_Absolute);
+		
+		URL.AddOption(*FString::Printf(TEXT("HandshakeId=%u"), ClientHandshakeId));
 
 		// Maybe this should be a control message since it changes the server setting (it's global driver setting).
 		URL.AddOption(TEXT("AutonomousAsSimulated"));
@@ -158,17 +179,19 @@ void UProxyListenerNotify::ConnectToGameServer(UNetConnection* ProxyConnection,
 		// it here again to point to the proxy.
 		GameServerConnectionState->NetDriver->Notify = GameServerConnectionState->GameServerNotify;
 
-		UE_LOG(LogNetProxy, Log, TEXT("Connecting to game server: %s:%s -> %s (player=%s)"), 
+		UE_LOG(LogNetProxy, Log, TEXT("Connecting to game server: %s:%s -> %s (player=%s client_handshake_id=%u)"), 
 			   *ProxyConnection->GetDriver()->GetName(), 
 			   *ProxyConnection->GetName(),
 			   *GameServerConnection->GetDriver()->GetName(),
-			   *NewPlayer->GetName());
+			   *NewPlayer->GetName(),
+			   ClientHandshakeId);
 	}
 	else
 	{
 		FGameServerSplitJoinRequest Request;
 		Request.Player = NewPlayer;
 		Request.Flags = Flags;
+		Request.ClientHandshakeId = ClientHandshakeId;
 
 		// The NMT_JoinSplit message can only be sent when the parent connection is open.
 		if (GameServerConnectionState->NetDriver->ServerConnection->GetConnectionState() == USOCK_Open)
@@ -176,39 +199,44 @@ void UProxyListenerNotify::ConnectToGameServer(UNetConnection* ProxyConnection,
 			GameServerConnectionState->PendingSplitJoinRequests.Add(Request);
 			ProxyNetDriver->FlushSplitJoinRequests(GameServerConnectionState);
 
-			UE_LOG(LogNetProxy, Log, TEXT("Connecting (multiplexed) to game server: %s:%s -> %s (player=%s"), 
+			UE_LOG(LogNetProxy, Log, TEXT("Connecting to game server (multiplexed): %s:%s -> %s (player=%s client_handshake_id=%u"), 
 				   *ProxyConnection->GetDriver()->GetName(), 
 				   *ProxyConnection->GetName(),
 				   *GameServerConnection->GetDriver()->GetName(),
-				   *NewPlayer->GetName());
+				   *NewPlayer->GetName(),
+				   Request.ClientHandshakeId);
 		}
 		else
 		{
 			GameServerConnectionState->PendingSplitJoinRequests.Add(Request);
 	
-			UE_LOG(LogNetProxy, Log, TEXT("Queuing up connection (multiplexed) to game server: %s:%s -> %s (player=%s)"), 
+			UE_LOG(LogNetProxy, Log, TEXT("Connecting to game server (queued multiplexed): %s:%s -> %s (player=%s client_handshake_id=%u)"), 
 				   *ProxyConnection->GetDriver()->GetName(),
 				   *ProxyConnection->GetName(),
 				   *GameServerConnection->GetDriver()->GetName(),
-				   *NewPlayer->GetName());
+				   *NewPlayer->GetName(),
+				   Request.ClientHandshakeId);
 		}
 	}
 
-	const int32 NewPlayerIndex = GameServerConnectionState->Players.Num();
 	GameServerConnectionState->Players.Add(NewPlayer);
 		
 	// Associate this proxy connection with the parent game server connection when beginning the handshake because
 	// the child connection hasn't been created yet. Once the handshake is complete, it's expected that this entry will
 	// be updated with the new child connection.
-	FGameServerAndPlayerIndex Key; 
-	Key.GameServerConnection = GameServerConnection;
-	Key.PlayerIndex = NewPlayerIndex;
-	ProxyNetDriver->GameToProxyConnection.Add(Key, ProxyConnection);
+	FMultiServerProxyInternalConnectionRoute Route;
+	Route.ProxyConnection = ProxyConnection;
+	Route.Player = NewPlayer;
+	Route.ParentGameServerConnection = GameServerConnection;
+
+	ProxyNetDriver->ClientHandshakeIdToInternalRoute.Add(ClientHandshakeId, Route);
 
 	if (!EnumHasAllFlags(Flags, UE::Net::EJoinFlags::NoPawn))
 	{
 		ProxyNetDriver->RegisterPrimaryGameServer(ProxyConnection, NewPlayer, GameServerConnectionStateIndex);
 	}
+
+	ProxyNetDriver->ProxyConnectionState[ProxyConnection] = EProxyConnectionState::ConnectingPrimary;
 }
 
 void UProxyBackendNetConnection::HandleClientPlayer(APlayerController* NewPlayerController, UNetConnection* GameServerConnection)
@@ -560,21 +588,30 @@ void UProxyNetDriver::ForwardRemoteFunction(UObject* RootObject, UObject* SubObj
 	AActor* OwningActor = Cast<AActor>(RootObject);
 	if (!ensure(OwningActor != nullptr))
 	{
-		UE_LOG(LogNetProxy, Verbose, TEXT("RPC: Ignoring %s on %s from proxy client because it doesn't have an owner."), *Function->GetName(), *RootObject->GetName());
+		UE_LOG(LogNetProxy, Warning, TEXT("RPC: Ignoring %s on %s from proxy connection because it doesn't have an owner."), *Function->GetName(), *RootObject->GetName());
 		return;
 	}
 
 	ULocalPlayer* OwningPlayer = Cast<ULocalPlayer>(OwningActor->GetNetOwningPlayerAnyRole());
 	if (OwningPlayer == nullptr)
 	{
-		UE_LOG(LogNetProxy, Verbose, TEXT("RPC: Ignoring %s on %s from proxy client because it owner %s doesn't have an owning player"), *Function->GetName(), *RootObject->GetName(), *OwningActor->GetName());
+		UE_LOG(LogNetProxy, Warning, TEXT("RPC: Ignoring %s on %s from proxy connection because it owner %s doesn't have an owning player"), *Function->GetName(), *RootObject->GetName(), *OwningActor->GetName());
 		return;
 	}
 
+	UE_LOG(LogNetProxy, VeryVerbose, TEXT("RPC: Pre-forwarding %s on %s from proxy connection to game server."), *Function->GetName(), *OwningActor->GetName());
+
 	if (TObjectPtr<APlayerController>* GamePlayerControllerPtr = ProxyPlayerToPrimaryGameServerPlayerController.Find(OwningPlayer))
 	{
 		UNetConnection* GameConnection = GamePlayerControllerPtr->Get()->NetConnection;
-		GameConnection->Driver->ProcessRemoteFunction(OwningActor, Function, Parms, static_cast<FOutParmRec*>(nullptr), static_cast<FFrame*>(nullptr), SubObject);
+		if (GameConnection)
+		{
+			GameConnection->Driver->ProcessRemoteFunction(OwningActor, Function, Parms, static_cast<FOutParmRec*>(nullptr), static_cast<FFrame*>(nullptr), SubObject);
+		}
+	}
+	else
+	{
+		UE_LOG(LogNetProxy, Warning, TEXT("RPC: Unable to forward %s on %s because player %s isn't mapped to a primary game server."), *Function->GetName(), *OwningActor->GetName(), *OwningPlayer->GetName());
 	}
 }
 
@@ -719,170 +756,180 @@ int32 UProxyNetDriver::GetNextGameServerClientId()
 	return NextGameServerClientId++;
 }
 
+uint32 UProxyNetDriver::GetNextClientHandshakeId()
+{
+	return NextClientHandshakeId++;
+}
+
 bool UProxyNetDriver::IsProxySpawned(AActor* Actor) const
 {
 	return ProxySpawnedActors.Contains(Actor);
 }
 
-void UProxyNetDriver::GameServerAssignPlayerController(UNetConnection* ChildGameServerConnection, UNetConnection* ParentGameServerConnection, APlayerController* GameServerPlayerController)
+void UProxyNetDriver::GameServerAssignPlayerController(UNetConnection* ChildGameServerConnection, UNetConnection* NotUsedParentGameServerConnection, APlayerController* GameServerPlayerController)
 {
-	UE_LOG(LogNetProxy, Log, TEXT("Completing handshake for connection %s:%s (parent:%s:%s) to primary game server using game server controller %s:%s."), 
-		   *ChildGameServerConnection->GetDriver()->GetName(), *ChildGameServerConnection->GetName(), 
-		   *ParentGameServerConnection->GetDriver()->GetName(), *ParentGameServerConnection->GetName(),
-		   *GameServerPlayerController->GetName(), *GameServerPlayerController->GetClass()->GetName());
-
-	// If a new player controller is spawned on the proxy from a game server it's assumed it will replace any
-	// existing game server player controllers used by connections and players on the proxy. For this reason
-	// always detach an existing game server player controllers from the connection.
-	DetachPlayerController(ChildGameServerConnection->PlayerController);
-	
-	ULocalPlayer* NewPlayer = GetPlayerOnGameServer(ParentGameServerConnection, GameServerPlayerController->NetPlayerIndex);
-	ensure(NewPlayer);
-
-	AttachPlayerController(GameServerPlayerController, ChildGameServerConnection, NewPlayer);
+	UE_LOG(LogNetProxy, Log, TEXT("Received a new player controller %s:%s for connection %s:%s (client_handshake_id=%u)."), 
+		   *GameServerPlayerController->GetName(), *GameServerPlayerController->GetClass()->GetName(),
+		   *ChildGameServerConnection->GetDriver()->GetName(), *ChildGameServerConnection->GetName(),
+		   GameServerPlayerController->GetClientHandshakeId());
 	
 	// The game server player controller is used as a client connection to the game server so must be ROLE_AutonomousProxy.
 	// NOTE: A proxy player controller will still be ROLE_Authority since it's spawned by the proxy.
 	GameServerPlayerController->SetRole(ROLE_AutonomousProxy);
 
+	const uint32 ClientHandshakeId = GameServerPlayerController->GetClientHandshakeId();
+	FMultiServerProxyInternalConnectionRoute* Route = ClientHandshakeIdToInternalRoute.Find(ClientHandshakeId);
+	if (Route == nullptr)
+	{
+		UE_LOG(LogNetProxy, Error, TEXT("Failed to find a mapping for game server connection request (client_handshake_id=%u)"), ClientHandshakeId);
+		return;
+	}
+	
+	UE_LOG(LogNetProxy, Log, TEXT("Found internal route (client_handshake_id=%u player=%s proxy_connection=%s:%s parent_game_server_connection=%s:%s)"), 
+		   ClientHandshakeId,
+		   *Route->Player->GetName(),
+		   *Route->ProxyConnection->GetDriver()->GetName(),
+		   *Route->ProxyConnection->GetName(),
+		   *Route->ParentGameServerConnection->GetDriver()->GetName(),
+		   *Route->ParentGameServerConnection->GetName());
+
+	ULocalPlayer* Player = Route->Player;
+		
 	// The player should only reference a proxy player controller if this game server connection
 	// is the primary game server for the proxy connection (decided later in this function).
-	NewPlayer->PlayerController = nullptr;
-
-	const int32 PlayerIndex = GameServerPlayerController->NetPlayerIndex;
-	FGameServerAndPlayerIndex Key;
-	Key.GameServerConnection = ParentGameServerConnection;
-	Key.PlayerIndex = PlayerIndex;
-
-	UE_LOG(LogNetProxy, VeryVerbose, TEXT("Lookup proxy connection associated with game server connection %s:%s and player index %d (game_server_controller=%s:%s"),
-		*ParentGameServerConnection->GetDriver()->GetName(),
-		*ParentGameServerConnection->GetName(),
-		PlayerIndex,
-		*GameServerPlayerController->GetName(),
-		*GameServerPlayerController->GetClass()->GetName());
+	Player->PlayerController = nullptr;
+	
+	// If a new player controller is spawned on the proxy from a game server it's assumed it will replace any
+	// existing game server player controllers used by connections and players on the proxy. For this reason
+	// always detach an existing game server player controllers from the connection.
+	DetachPlayerController(ChildGameServerConnection->PlayerController);
+	AttachPlayerController(GameServerPlayerController, ChildGameServerConnection, Player);
 
 	// Find the proxy connection that established this connection to the game server.
-	TObjectPtr<UNetConnection>* ProxyConnectionPtr = GameToProxyConnection.Find(Key);
-	if (ensure(ProxyConnectionPtr))
+	TObjectPtr<UNetConnection> ProxyConnection = Route->ProxyConnection;
+
+	if (ensure(ProxyConnectionState.Contains(ProxyConnection)))
 	{
-		TObjectPtr<UNetConnection> ProxyConnection = *ProxyConnectionPtr;
+		EProxyConnectionState State = ProxyConnectionState[ProxyConnection];
 
-		bool bBindProxyConnectionToPrimaryGameServer = false;
+		UE_LOG(LogNetProxy, Log, TEXT("Attempting to assign game server player controller %s to proxy connection %s:%s (state=%s)"),
+			*GameServerPlayerController->GetName(),
+			*ProxyConnection->GetDriver()->GetName(),
+			*ProxyConnection->GetName(),
+			*LexToString(State));
 
-		if (int32* PlayerPrimaryGameServerIdx = ProxyPlayerToPrimaryGameServer.Find(NewPlayer))
+		switch (State)
 		{
-			if (ensure(*PlayerPrimaryGameServerIdx < GameServerConnections.Num()))
-			{
-				FGameServerConnectionState& PlayerPrimaryGameServer = GameServerConnections[*PlayerPrimaryGameServerIdx];
-				if (PlayerPrimaryGameServer.NetDriver->ServerConnection == ParentGameServerConnection)
+			case EProxyConnectionState::Disconnected:
+				break;
+
+			case EProxyConnectionState::ConnectingPrimary:
 				{
-					bBindProxyConnectionToPrimaryGameServer = true;
+					if (int32* PlayerPrimaryGameServerIdx = ProxyPlayerToPrimaryGameServer.Find(Player))
+					{
+						ensure(ProxyConnection->PlayerController == nullptr);
+
+						if (ensure(*PlayerPrimaryGameServerIdx < GameServerConnections.Num()))
+						{
+							FGameServerConnectionState& PlayerPrimaryGameServer = GameServerConnections[*PlayerPrimaryGameServerIdx];
+							if (PlayerPrimaryGameServer.NetDriver->ServerConnection == Route->ParentGameServerConnection)
+							{
+								UPlayer* PreviousPlayer = GameServerPlayerController->Player;
+			
+								APlayerController* ProxyPlayerController = SpawnPlayerController();
+								if (ensure(ProxyPlayerController))
+								{
+									DetachAndDestroyPlayerController(ProxyConnection->PlayerController);
+									AttachPlayerController(ProxyPlayerController, ProxyConnection, Player);
+
+									// The player always points to the proxy player controller if that player is
+									// considered the primary player for the proxy connection.
+									Player->PlayerController = ProxyPlayerController;
+
+									ProxyPlayerToPrimaryGameServerPlayerController.Add(Player, GameServerPlayerController);
+
+									UE_LOG(LogNetProxy, Log, TEXT("Successfully connected proxy connection %s:%s to primary game server connection %s:%s (game_server_controller=%s:%s proxy_controller=%s:%s player=%s client_handshake_id=%u)"),
+										*ProxyConnection->GetDriver()->GetName(),
+										*ProxyConnection->GetName(),
+										*ChildGameServerConnection->GetDriver()->GetName(),
+										*ChildGameServerConnection->GetName(),
+										*GameServerPlayerController->GetName(),
+										*GameServerPlayerController->GetClass()->GetName(),
+										*ProxyPlayerController->GetName(),
+										*ProxyPlayerController->GetClass()->GetName(),
+										*Player->GetName(),
+										ClientHandshakeId);
+
+									ProxyConnectionState[ProxyConnection] = EProxyConnectionState::ConnectedPrimary;
+								}
+							}
+						}
+					}
+					else
+					{
+						UE_LOG(LogNetProxy, Log, TEXT("Game server player controller %s not from proxy connection %s:%s primary game server"),
+							*GameServerPlayerController->GetName(),
+							*ProxyConnection->GetDriver()->GetName(),
+							*ProxyConnection->GetName());
+					}
 				}
-			}
-		}
+				break;
 
-		// If the proxy connection has already been bound to a primary game server player controller, and the new player has the 
-		// same player id as the existing player id, then change the primary game server and rebind the two connections.
-		/*else if(ProxyConnection->PlayerController)
-		{
-			if (ULocalPlayer* ExistingPlayer = Cast<ULocalPlayer>(ProxyConnection->PlayerController->Player))
-			{
-				if (ProxyPlayerToPrimaryGameServerPlayerController.Contains(ExistingPlayer))
+			case EProxyConnectionState::ConnectedPrimary:
 				{
-					int32* PreviousPrimaryGameServerIdx = ProxyPlayerToPrimaryGameServer.Find(ExistingPlayer);
-					if (ensure(PreviousPrimaryGameServerIdx))
+					if (!GameServerPlayerController->IsA(ANoPawnPlayerController::StaticClass()))
 					{
-						if (NewPlayer->GetCachedUniqueNetId() == ExistingPlayer->GetCachedUniqueNetId())
+						if (ULocalPlayer* ExistingPrimaryPlayer = Cast<ULocalPlayer>(ProxyConnection->PlayerController->Player))
 						{
-							int32 NewPrimaryGameServerIdx = GetGameServerWithPlayer(NewPlayer);
-
-							UE_LOG(LogNetProxy, Log, TEXT("Changing primary game server for proxy connection %s:%s from player %s to %s (prev_game_server_idx=%d new_game_server_idx=%d)."),
-								   *ProxyConnection->Driver->GetName(),
-								   *ProxyConnection->GetName(),
-								   *ExistingPlayer->GetName(),
-								   *NewPlayer->GetName(),
-								   *PreviousPrimaryGameServerIdx,
-								   NewPrimaryGameServerIdx);
+							if (ProxyPlayerToPrimaryGameServerPlayerController.Contains(ExistingPrimaryPlayer))
+							{
+								int32* PreviousPrimaryGameServerIdx = ProxyPlayerToPrimaryGameServer.Find(ExistingPrimaryPlayer);
+								if (ensure(PreviousPrimaryGameServerIdx))
+								{
+									if (Player->GetCachedUniqueNetId() == ExistingPrimaryPlayer->GetCachedUniqueNetId())
+									{
+										int32 NewPrimaryGameServerIdx = GetGameServerWithPlayer(Player);
+
+										UE_LOG(LogNetProxy, Log, TEXT("Changing primary game server for proxy connection %s:%s from player %s to %s (prev_game_server_idx=%d new_game_server_idx=%d)."),
+											*ProxyConnection->Driver->GetName(),
+											*ProxyConnection->GetName(),
+											*ExistingPrimaryPlayer->GetName(),
+											*Player->GetName(),
+											*PreviousPrimaryGameServerIdx,
+											NewPrimaryGameServerIdx);
 					
-							ProxyPlayerToPrimaryGameServerPlayerController.Remove(ExistingPlayer);
+										ProxyPlayerToPrimaryGameServerPlayerController.Remove(ExistingPrimaryPlayer);
 
-							DeregisterPrimaryGameServer(ProxyConnection);
-							RegisterPrimaryGameServer(ProxyConnection, NewPlayer, NewPrimaryGameServerIdx);
+										DeregisterPrimaryGameServer(ProxyConnection);
+										RegisterPrimaryGameServer(ProxyConnection, Player, NewPrimaryGameServerIdx);
 				
-							bBindProxyConnectionToPrimaryGameServer = true;
+										ExistingPrimaryPlayer->PlayerController = nullptr;
+										Player->PlayerController = ProxyConnection->PlayerController;
+										ProxyConnection->PlayerController->Player = Player;
+
+										ProxyPlayerToPrimaryGameServerPlayerController.Add(Player, GameServerPlayerController);
+									}
+								}
+							}
 						}
 					}
 				}
-			}
-		}*/
-
-		if (bBindProxyConnectionToPrimaryGameServer)
-		{
-			UPlayer* PreviousPlayer = GameServerPlayerController->Player;
-			
-			APlayerController* ProxyPlayerController = SpawnPlayerController();
+				break;
 
-			if (ProxyPlayerController)
-			{
-				DetachAndDestroyPlayerController(ProxyConnection->PlayerController);
-				AttachPlayerController(ProxyPlayerController, ProxyConnection, NewPlayer);
-
-				// The player always points to the proxy player controller if that player is
-				// considered the primary player for the proxy connection.
-				NewPlayer->PlayerController = ProxyPlayerController;
-
-				ProxyPlayerToPrimaryGameServerPlayerController.Add(NewPlayer, GameServerPlayerController);
-
-				UE_LOG(LogNetProxy, Log, TEXT("Binding proxy connection %s:%s to primary game server connection %s:%s (game_server_controller=%s:%s proxy_controller=%s:%s player=%s:%d)"),
-					*ProxyConnection->GetDriver()->GetName(),
-					*ProxyConnection->GetName(),
-					*ChildGameServerConnection->GetDriver()->GetName(),
-					*ChildGameServerConnection->GetName(),
-					*GameServerPlayerController->GetName(),
-					*GameServerPlayerController->GetClass()->GetName(),
-					*ProxyPlayerController->GetName(),
-					*ProxyPlayerController->GetClass()->GetName(),
-					*NewPlayer->GetName(),
-					PlayerIndex);
-			}
-			else
-			{
-				UE_LOG(LogNetProxy, Error, TEXT("Failed to bind a proxy connection %s:%s to a new primary game server %s:%s."),
-					*ProxyConnection->GetDriver()->GetName(),
-					*ProxyConnection->GetName(),
-					*ChildGameServerConnection->GetDriver()->GetName(),
-					*ChildGameServerConnection->GetName());
-			}
-		}
-		else
-		{
-			UE_LOG(LogNetProxy, Log, TEXT("Ignoring connection %s:%s to non-primary game server %s:%s (game_server_controller=%s:%s player=%s:%d)"),
-				*ProxyConnection->GetDriver()->GetName(),
-				*ProxyConnection->GetName(),
-				*ChildGameServerConnection->GetDriver()->GetName(),
-				*ChildGameServerConnection->GetName(),
-				*GameServerPlayerController->GetName(),
-				*GameServerPlayerController->GetClass()->GetName(),
-				*NewPlayer->GetName(),
-				PlayerIndex);
 		}
-			
-		// Update the mapping of proxy connection to game server connection to use the new child game server connection (if it's 
-		// different from the parent connection).
-		FGameServerAndPlayerIndex NewKey;
-		NewKey.GameServerConnection = ChildGameServerConnection;
-		NewKey.PlayerIndex = PlayerIndex;
-		GameToProxyConnection.Remove(Key);
-		GameToProxyConnection.Add(NewKey, ProxyConnection);
 	}
+			
+	// Update the internal connection mapping to use the new child game server connection (if it's different from the parent connection).
+	ClientHandshakeIdToInternalRoute.Remove(NextClientHandshakeId);
+	ClientHandshakeIdToInternalRoute.Add(NextClientHandshakeId, *Route);
 
 	// If this was a parent connection send any join requests for any pending multiplexed connections.
-	if (ChildGameServerConnection == ParentGameServerConnection)
+	if (ChildGameServerConnection == Route->ParentGameServerConnection)
 	{
 		FGameServerConnectionState* CurrGameServerConnectionState = nullptr;
 		for (FGameServerConnectionState& GameServerConnectionState : GameServerConnections)
 		{
-			if (GameServerConnectionState.NetDriver->ServerConnection == ParentGameServerConnection)
+			if (GameServerConnectionState.NetDriver->ServerConnection == Route->ParentGameServerConnection)
 			{
 				CurrGameServerConnectionState = &GameServerConnectionState;
 				break;
@@ -912,13 +959,14 @@ void UProxyNetDriver::FlushSplitJoinRequests(FGameServerConnectionState* GameSer
 
 	for (FGameServerSplitJoinRequest& Request : GameServerConnectionState->PendingSplitJoinRequests)
 	{
-		TArray<FString> Options;
-
-		UE_LOG(LogNetProxy, Log, TEXT("Sending queued connection (multiplexed) request to game server: %s (player=%s flags=%d)"),
+		UE_LOG(LogNetProxy, Log, TEXT("Sending queued connection (multiplexed) request to game server: %s (player=%s flags=%d client_handshake_id=%u)"),
 			   *GameServerConnectionState->NetDriver->GetName(),
 			   *Request.Player->GetName(),
-			   Request.Flags);
+			   Request.Flags,
+			   Request.ClientHandshakeId);
 
+		TArray<FString> Options;
+		Options.Add(FString::Printf(TEXT("HandshakeId=%u"), Request.ClientHandshakeId));
 		Request.Player->SendSplitJoin(Options, GameServerConnectionState->NetDriver, Request.Flags);
 	}
 
@@ -987,6 +1035,12 @@ void UProxyNetDriver::AttachPlayerController(APlayerController* PlayerController
 	// This function doesn't associate the player with a player controller since it differs between
 	// proxy and game server player controllers.
 
+	UE_LOG(LogNetProxy, Log, TEXT("Attaching controller %s:%s on connection %s and player %s."),
+		   *PlayerController->GetName(),
+		   *PlayerController->GetClass()->GetName(),
+		   *Connection->GetName(),
+		   *Player->GetName());
+
 	Connection->SetConnectionState(EConnectionState::USOCK_Open);
 
 	// It's assumed the connection is not associated with a player controller.	
@@ -1008,10 +1062,10 @@ void UProxyNetDriver::DetachPlayerController(APlayerController* PlayerController
 	{
 		UPlayer* Player = PlayerController->Player;
 
-		UE_LOG(LogNetProxy, Log, TEXT("Destroying old controller %s:%s on %s attached to player %s."),
+		UE_LOG(LogNetProxy, Log, TEXT("Detaching old controller %s:%s on %s attached to player %s."),
 			*PlayerController->GetName(),
 			*PlayerController->GetClass()->GetName(),
-			*this->GetName(),
+			*GetName(),
 			*Player->GetName());
 
 		if (UNetConnection* Connection = PlayerController->NetConnection)
@@ -1087,24 +1141,6 @@ void UProxyNetDriver::DestroyPlayerController(APlayerController* PlayerControlle
 	}
 }
 
-ULocalPlayer* UProxyNetDriver::GetPlayerOnGameServer(UNetConnection* ParentGameServerConnection, int32 PlayerIndex)
-{
-	ULocalPlayer* FoundPlayer = nullptr;
-
-	for (FGameServerConnectionState& GameServerConnectionState : GameServerConnections)
-	{
-		if (GameServerConnectionState.NetDriver == ParentGameServerConnection->GetDriver())
-		{
-			if (ensure(PlayerIndex < GameServerConnectionState.Players.Num()))
-			{
-				FoundPlayer = GameServerConnectionState.Players[PlayerIndex];
-			}
-		}
-	}
-
-	return FoundPlayer;
-}
-
 UNetConnection* UProxyNetDriver::GetProxyConnectionFromPrimaryPlayer(UPlayer* Player)
 {
 	if (ULocalPlayer* LocalPlayer = Cast<ULocalPlayer>(Player))
diff --git a/Engine/Plugins/Runtime/MultiServerReplication/Source/MultiServerReplication/Public/MultiServerProxy.h b/Engine/Plugins/Runtime/MultiServerReplication/Source/MultiServerReplication/Public/MultiServerProxy.h
--- a/Engine/Plugins/Runtime/MultiServerReplication/Source/MultiServerReplication/Public/MultiServerProxy.h
+++ b/Engine/Plugins/Runtime/MultiServerReplication/Source/MultiServerReplication/Public/MultiServerProxy.h
@@ -9,6 +9,15 @@
 #include "IpConnection.h"
 #include "MultiServerProxy.generated.h"
 
+enum class EProxyConnectionState : uint8
+{
+	Disconnected,
+	ConnectingPrimary,
+	ConnectedPrimary
+};
+
+FString LexToString(EProxyConnectionState State);
+
 USTRUCT()
 struct FGameServerSplitJoinRequest
 {
@@ -18,6 +27,8 @@ struct FGameServerSplitJoinRequest
 	TObjectPtr<class ULocalPlayer> Player;
 
 	UE::Net::EJoinFlags Flags;
+
+	uint32 ClientHandshakeId;
 };
 
 /**
@@ -44,24 +55,18 @@ struct FGameServerConnectionState
 };
 
 USTRUCT()
-struct FGameServerAndPlayerIndex
+struct FMultiServerProxyInternalConnectionRoute
 {
 	GENERATED_BODY()
 
-	bool operator==(const FGameServerAndPlayerIndex& Rhs) const
-	{
-		return GetTypeHash(Rhs) == GetTypeHash(*this);
-	}
-
-	friend uint32 GetTypeHash(const FGameServerAndPlayerIndex& Obj)
-	{
-		return HashCombine(::GetTypeHash(Obj.PlayerIndex), ::GetTypeHash(Obj.GameServerConnection));
-	}
+	UPROPERTY()
+	TObjectPtr<UNetConnection> ProxyConnection;
 
 	UPROPERTY()
-	TObjectPtr<UNetConnection> GameServerConnection;
+	TObjectPtr<UNetConnection> ParentGameServerConnection;
 
-	int32 PlayerIndex;
+	UPROPERTY()
+	TObjectPtr<class ULocalPlayer> Player;
 };
 
 /** 
@@ -271,12 +276,12 @@ public:
 	virtual void InternalProcessRemoteFunction(class AActor* Actor, class UObject* SubObject, class UNetConnection* Connection, class UFunction* Function, void* Parms, FOutParmRec* OutParms, FFrame* Stack, bool bIsServer) override;
 	virtual int32 ServerReplicateActors(float DeltaSeconds) override;
 	virtual bool CanDowngradeActorRole(UNetConnection* ProxyConnection, AActor* Actor) const override;
-
-	/** Get the next identifier for incoming client connections to the proxy. */
-	int32 GetNextProxyClientId();
 	
 	/** Get the next identifier for outgoing connections to game servers. */
 	int32 GetNextGameServerClientId();
+
+	/** Get the next connection handshake id. */
+	uint32 GetNextClientHandshakeId();
 	
 	/** Called when the player controller associated with a connection is changed (either at the end of initial connection handshake, or if changed after successfully connected). */
 	void GameServerAssignPlayerController(UNetConnection* ChildGameServerConnection, UNetConnection* ParentGameServerConnection, APlayerController* PlayerController);
@@ -310,8 +315,10 @@ private:
 	/** Attach a new player, player controller and connection. */
 	void AttachPlayerController(APlayerController* PlayerController, UNetConnection* Connection, ULocalPlayer* Player);
 
-	/** Detach an existing player, player controller and connection and (optionally) destroy the controller. */
+	/** Detach an existing player, player controller and connection. */
 	void DetachPlayerController(APlayerController* PlayerController);
+	
+	/** Detach an existing player, player controller and connection and destroy the controller. */
 	void DetachAndDestroyPlayerController(APlayerController* PlayerController);
 
 	/** Create and destroy a proxy player controller. Returns null if failed to spawn the controller. */
@@ -320,27 +327,27 @@ private:
 	/** Destroy a player controller spawned by the proxy. */
 	void DestroyPlayerController(APlayerController* PlayerController);
 	
-	/** Return the local player associated with a player index on a given game server. */
-	class ULocalPlayer* GetPlayerOnGameServer(UNetConnection* ParentGameServerConnection, int32 PlayerIndex);
-
 	/** Return a proxy connection from a primary player (a player on a primary game server). */
 	static UNetConnection* GetProxyConnectionFromPrimaryPlayer(UPlayer* Player);
 
 	/** Return a proxy player controller given a primary game server controller. */
 	APlayerController* GetProxyControllerFromPrimaryGameServerController(APlayerController* GameServerPlayerController);
 	
-	/** Associate a connection to a game server with the proxy connection that established that connection. */
+	/** Associate a request to connect to a game server with a route from a proxy connection to game server. */
 	UPROPERTY()
-	TMap<FGameServerAndPlayerIndex, TObjectPtr<UNetConnection>> GameToProxyConnection;
+	TMap<uint32, FMultiServerProxyInternalConnectionRoute> ClientHandshakeIdToInternalRoute;
 
-	/** Map a player to it's primary game server index into GameServerConnections (set when connection handshake begins). */
+	/** Map a player to it's primary game server index into GameServerConnections. */
 	UPROPERTY()
 	TMap<TObjectPtr<class ULocalPlayer>, int32> ProxyPlayerToPrimaryGameServer;
 
-	/** Map a player to it's primary game server player controller (set after connection handshake is completed successfully). */
+	/** Map a player to it's primary game server player controller. */
 	UPROPERTY()
 	TMap<TObjectPtr<class ULocalPlayer>, TObjectPtr<APlayerController>> ProxyPlayerToPrimaryGameServerPlayerController;
 
+	/** The state of each incoming proxy connection. */
+	TMap<TObjectPtr<UNetConnection>, EProxyConnectionState> ProxyConnectionState;
+
 	/** Net drivers and associated state used to connect to backend game servers. */
 	UPROPERTY()
 	TArray<FGameServerConnectionState> GameServerConnections;
@@ -352,9 +359,11 @@ private:
 	/** A set of actors that have been spawned by the proxy. */
 	UPROPERTY()
 	TSet<TObjectPtr<AActor>> ProxySpawnedActors;
-
+	
 	int32 NextGameServerClientId = 0;
 	
+	uint32 NextClientHandshakeId = 123;
+	
 	// The primary game server to use for the next primary game client.
 	int32 PrimaryGameServerForNextClient = 0;
 
diff --git a/Engine/Source/Runtime/Core/Public/Misc/EngineNetworkCustomVersion.h b/Engine/Source/Runtime/Core/Public/Misc/EngineNetworkCustomVersion.h
--- a/Engine/Source/Runtime/Core/Public/Misc/EngineNetworkCustomVersion.h
+++ b/Engine/Source/Runtime/Core/Public/Misc/EngineNetworkCustomVersion.h
@@ -93,6 +93,7 @@ struct FEngineNetworkCustomVersion
 		RemoteObjectReferences = 38,			// Bump version to support serialization changes for remote object references 
 		BeaconNetIDVariant = 39,				// Bump version for online beacon's BeaconAssignGUID control message switching to use FNetIDVariant
 		JoinNoPawn = 40,						// Bump version to support NMT_JoinNoPawn and NMT_JoinNoPawnSplit control messages.
+		ClientHandshakeId = 41,					// Bump version to support serialization of APlayerController::ClientHandeshakeId.
 
 		// -----<new versions can be added above this line>-------------------------------------------------
 		VersionPlusOne,
diff --git a/Engine/Source/Runtime/CoreUObject/Private/UObject/RemoteObject.cpp b/Engine/Source/Runtime/CoreUObject/Private/UObject/RemoteObject.cpp
--- a/Engine/Source/Runtime/CoreUObject/Private/UObject/RemoteObject.cpp
+++ b/Engine/Source/Runtime/CoreUObject/Private/UObject/RemoteObject.cpp
@@ -331,7 +331,12 @@ void ChangeOwnerServerId(const UObject* Object, FRemoteServerId NewOwnerServerId
 	
 #if UE_WITH_REMOTE_OBJECT_HANDLE
 	UE::RemoteObject::Handle::FRemoteObjectStub* RemoteStub = FindRemoteObjectStub(FObjectHandleUtils::GetRemoteId(Object));
-	RemoteStub->OwningServerId = NewOwnerServerId;
+
+	// The remote stub is always expected to be found for this object.
+	if (ensure(RemoteStub))
+	{
+		RemoteStub->OwningServerId = NewOwnerServerId;
+	}
 #endif
 }
 
diff --git a/Engine/Source/Runtime/Engine/Classes/Engine/NetConnection.h b/Engine/Source/Runtime/Engine/Classes/Engine/NetConnection.h
--- a/Engine/Source/Runtime/Engine/Classes/Engine/NetConnection.h
+++ b/Engine/Source/Runtime/Engine/Classes/Engine/NetConnection.h
@@ -412,10 +412,19 @@ public:
 	/** Clears the actor starvation map */
 	void ResetActorsStarvedByClassTimeMap() { ActorsStarvedByClassTimeMap.Empty(); }
 
+	/** Get the client handshake id for this connection. */
+	ENGINE_API uint32 GetClientHandshakeId() const;
+
+	/** Set the client handshake id for this connection. */
+	ENGINE_API void SetClientHandshakeId(uint32 InClientHandshakeId);
+
 private:
 	// Connection information.
 	EConnectionState	State;					// State this connection is in.
 
+	/** An optional unique identifier set by the client when establishing a connection. */
+	uint32 ClientHandshakeId = 0;
+
 public:
 	ENGINE_API const EConnectionState GetConnectionState() const;
 	ENGINE_API void SetConnectionState(EConnectionState ConnectionState);
diff --git a/Engine/Source/Runtime/Engine/Classes/Engine/NetDriver.h b/Engine/Source/Runtime/Engine/Classes/Engine/NetDriver.h
--- a/Engine/Source/Runtime/Engine/Classes/Engine/NetDriver.h
+++ b/Engine/Source/Runtime/Engine/Classes/Engine/NetDriver.h
@@ -2352,6 +2352,9 @@ public:
 	/** Return true if the driver supports NMT_JoinNoPawn and NMT_JoinNoPawnSplit. */
 	bool HasSupportForNoPawnConnection();
 
+	/** Enumerate all connections to find one that matches a particular player id. */
+	TObjectPtr<UNetConnection> FindConnection(struct FUniqueNetIdRepl& PlayerId);
+
 protected:
 	
 	/** Stream of random numbers to be used by this instance of UNetDriver */
diff --git a/Engine/Source/Runtime/Engine/Classes/GameFramework/Controller.h b/Engine/Source/Runtime/Engine/Classes/GameFramework/Controller.h
--- a/Engine/Source/Runtime/Engine/Classes/GameFramework/Controller.h
+++ b/Engine/Source/Runtime/Engine/Classes/GameFramework/Controller.h
@@ -227,7 +227,7 @@ public:
 	//~ End AActor Interface
 
 	/** Getter for Pawn */
-	FORCEINLINE APawn* GetPawn() const { return Pawn; }
+	FORCEINLINE TObjectPtr<APawn> GetPawn() const { return Pawn; }
 
 	/** Templated version of GetPawn, will return nullptr if cast fails */
 	template<class T>
diff --git a/Engine/Source/Runtime/Engine/Classes/GameFramework/PlayerController.h b/Engine/Source/Runtime/Engine/Classes/GameFramework/PlayerController.h
--- a/Engine/Source/Runtime/Engine/Classes/GameFramework/PlayerController.h
+++ b/Engine/Source/Runtime/Engine/Classes/GameFramework/PlayerController.h
@@ -17,6 +17,7 @@
 #include "GameFramework/Controller.h"
 #include "UObject/TextProperty.h"
 #include "GameFramework/PlayerMuteList.h"
+#include "GameFramework/OnlineReplStructs.h"
 #include "Camera/PlayerCameraManager.h"
 #include "GameFramework/ForceFeedbackParameters.h"
 #include "GameFramework/UpdateLevelVisibilityLevelInfo.h"
@@ -1815,6 +1816,9 @@ public:
 	//~ Begin UObject Interface
 	ENGINE_API virtual void PostLoad() override;
 	ENGINE_API virtual void GetLifetimeReplicatedProps(TArray< FLifetimeProperty > & OutLifetimeProps) const override;
+#if UE_WITH_REMOTE_OBJECT_HANDLE
+	ENGINE_API virtual void PostMigrate(const FUObjectMigrationContext& MigrationContext) override;
+#endif
 	//~ End UObject Interface
 
 	//~ Begin AActor Interface
@@ -1845,6 +1849,11 @@ public:
 	ENGINE_API virtual void PostInitializeComponents() override;
 	ENGINE_API virtual void EnableInput(class APlayerController* PlayerController) override;
 	ENGINE_API virtual void DisableInput(class APlayerController* PlayerController) override;
+
+	ENGINE_API virtual void Serialize(FArchive& Ar) override;
+#if UE_WITH_REMOTE_OBJECT_HANDLE
+	ENGINE_API virtual void Tick(float DeltaTime) override;
+#endif
 protected:
 	ENGINE_API virtual void BeginPlay() override;
 	//~ End AActor Interface
@@ -2301,6 +2310,24 @@ public:
 	/** True if the NetworkPhysicsTickOffset is setup, before this is setup the GetNetworkPhysicsTickOffset() and GetPhysicsTimestamp() functions will not return valid results */
 	bool GetNetworkPhysicsTickOffsetAssigned() const { return bNetworkPhysicsTickOffsetAssigned; }
 
+	/** Return the handshake cookie of the connection this player controller is associated. */
+	ENGINE_API uint32 GetClientHandshakeId() const;
+
+	/** Associate this player controller with a connection using it's handshake cookie. */
+	ENGINE_API void SetClientHandshakeId(uint32 InClientHandshakeId);
+
+private:
+	
+	// This is a server-specific variable so it should not be migrated between servers and remain non-UPROPERTY().
+	bool bSendClientRestartRPC = false;
+
+	/** The player id associated with this player controller. */
+	UPROPERTY()
+	FUniqueNetIdRepl CachedConnectionPlayerId;
+
+	/** An optional client handshake id used by the connection associated with this player controller. */
+	UPROPERTY(DuplicateTransient)
+	uint32 ClientHandshakeId = 0;
 };
 
 /** 
diff --git a/Engine/Source/Runtime/Engine/Private/Actor.cpp b/Engine/Source/Runtime/Engine/Private/Actor.cpp
--- a/Engine/Source/Runtime/Engine/Private/Actor.cpp
+++ b/Engine/Source/Runtime/Engine/Private/Actor.cpp
@@ -1333,6 +1333,11 @@ void AActor::PostMigrate(const FUObjectMigrationContext& MigrationContext)
 
 			World->AddNetworkActor(this);
 		}
+		else
+		{
+			// Make sure that all tick functions have been deregistered if we aren't acquring ownership.
+			RegisterAllActorTickFunctions(false, false);
+		}
 
 		ActorHasBegunPlay = EActorBeginPlayState::HasBegunPlay;
 	}
diff --git a/Engine/Source/Runtime/Engine/Private/LevelActor.cpp b/Engine/Source/Runtime/Engine/Private/LevelActor.cpp
--- a/Engine/Source/Runtime/Engine/Private/LevelActor.cpp
+++ b/Engine/Source/Runtime/Engine/Private/LevelActor.cpp
@@ -1080,6 +1080,11 @@ APlayerController* UWorld::SpawnPlayActor(UPlayer* NewPlayer, ENetRole RemoteRol
 			return NULL;
 		}
 
+		if (UNetConnection* Connection = Cast<UNetConnection>(NewPlayer))
+		{
+			NewPlayerController->SetClientHandshakeId(Connection->GetClientHandshakeId());
+		}
+
 		UE_LOG(LogSpawn, Log, TEXT("%s got player %s [%s]"), *NewPlayerController->GetName(), *NewPlayer->GetName(), UniqueId.IsValid() ? *UniqueId->ToString() : TEXT("Invalid"));
 
 		// Possess the newly-spawned player.
diff --git a/Engine/Source/Runtime/Engine/Private/LevelTick.cpp b/Engine/Source/Runtime/Engine/Private/LevelTick.cpp
--- a/Engine/Source/Runtime/Engine/Private/LevelTick.cpp
+++ b/Engine/Source/Runtime/Engine/Private/LevelTick.cpp
@@ -573,11 +573,14 @@ void UWorld::ProcessLevelStreamingVolumes(FVector* OverrideViewLocation)
 	bool bStreamingVolumesAreRelevant = false;
 	for (FConstPlayerControllerIterator Iterator = GetPlayerControllerIterator(); Iterator; ++Iterator)
 	{
-		APlayerController* PlayerActor = Iterator->Get();
-		if (PlayerActor && PlayerActor->bIsUsingStreamingVolumes)
+		if (!Iterator->IsRemote())
 		{
-			bStreamingVolumesAreRelevant = true;
-			break;
+			APlayerController* PlayerActor = Iterator->Get();
+			if (PlayerActor && PlayerActor->bIsUsingStreamingVolumes)
+			{
+				bStreamingVolumesAreRelevant = true;
+				break;
+			}
 		}
 	}
 
@@ -615,6 +618,11 @@ void UWorld::ProcessLevelStreamingVolumes(FVector* OverrideViewLocation)
 		// volumes that contain player viewpoints.
 		for( FConstPlayerControllerIterator Iterator = GetPlayerControllerIterator(); Iterator; ++Iterator )
 		{
+			if (Iterator->IsRemote())
+			{
+				continue;
+			}
+			
 			APlayerController* PlayerActor = Iterator->Get();
 			if (PlayerActor && PlayerActor->bIsUsingStreamingVolumes)
 			{
@@ -1573,6 +1581,11 @@ void UWorld::Tick( ELevelTick TickType, float DeltaSeconds )
 				// Update cameras last. This needs to be done before NetUpdates, and after all actors have been ticked.
 				for( FConstPlayerControllerIterator Iterator = GetPlayerControllerIterator(); Iterator; ++Iterator )
 				{
+					if (Iterator->IsRemote())
+					{
+						continue;
+					}
+					
 					if (APlayerController* PlayerController = Iterator->Get())
 					{
 						if (!bIsPaused || PlayerController->ShouldPerformFullTickWhenPaused())
diff --git a/Engine/Source/Runtime/Engine/Private/LocalPlayer.cpp b/Engine/Source/Runtime/Engine/Private/LocalPlayer.cpp
--- a/Engine/Source/Runtime/Engine/Private/LocalPlayer.cpp
+++ b/Engine/Source/Runtime/Engine/Private/LocalPlayer.cpp
@@ -384,7 +384,7 @@ void ULocalPlayer::SendSplitJoin(TArray<FString>& Options, UNetDriver* NetDriver
 
 			if (EnumHasAllFlags(Flags, UE::Net::EJoinFlags::NoPawn))
 			{
-				FNetControlMessage<NMT_JoinNoPawnSplit>::Send(NetDriver->ServerConnection, UniqueIdRepl);
+				FNetControlMessage<NMT_JoinNoPawnSplit>::Send(NetDriver->ServerConnection, URLString, UniqueIdRepl);
 			}
 			else
 			{
diff --git a/Engine/Source/Runtime/Engine/Private/NetConnection.cpp b/Engine/Source/Runtime/Engine/Private/NetConnection.cpp
--- a/Engine/Source/Runtime/Engine/Private/NetConnection.cpp
+++ b/Engine/Source/Runtime/Engine/Private/NetConnection.cpp
@@ -1025,6 +1025,16 @@ void UNetConnection::Serialize( FArchive& Ar )
 	}
 }
 
+uint32 UNetConnection::GetClientHandshakeId() const
+{
+	return ClientHandshakeId;
+}
+
+void UNetConnection::SetClientHandshakeId(uint32 InClientHandshakeId)
+{
+	ClientHandshakeId = InClientHandshakeId;
+}
+
 PRAGMA_DISABLE_DEPRECATION_WARNINGS
 EConnectionState const UNetConnection::GetConnectionState() const
 {
diff --git a/Engine/Source/Runtime/Engine/Private/NetDriver.cpp b/Engine/Source/Runtime/Engine/Private/NetDriver.cpp
--- a/Engine/Source/Runtime/Engine/Private/NetDriver.cpp
+++ b/Engine/Source/Runtime/Engine/Private/NetDriver.cpp
@@ -801,6 +801,25 @@ bool UNetDriver::HasSupportForNoPawnConnection()
 	return GbEnableSupportForNoPawnConnection || bSupportNoPawnConnection;
 }
 
+TObjectPtr<UNetConnection> UNetDriver::FindConnection(FUniqueNetIdRepl& PlayerId)
+{
+	UNetConnection* FoundConnection = nullptr;
+
+	// Find the connection with the same PlayerId.
+	ForEachClientConnection([PlayerId, &FoundConnection](UNetConnection* ClientConnection)
+	{
+		if (ClientConnection->PlayerId == PlayerId)
+		{
+			FoundConnection = ClientConnection;
+			return true;
+		}
+
+		return false;
+	});
+
+	return FoundConnection;
+}
+
 void UNetDriver::AssertValid()
 {
 }
diff --git a/Engine/Source/Runtime/Engine/Private/PlayerController.cpp b/Engine/Source/Runtime/Engine/Private/PlayerController.cpp
--- a/Engine/Source/Runtime/Engine/Private/PlayerController.cpp
+++ b/Engine/Source/Runtime/Engine/Private/PlayerController.cpp
@@ -75,6 +75,10 @@
 #include "Net/Iris/ReplicationSystem/ReplicationSystemUtil.h"
 #endif // UE_WITH_IRIS
 
+#if UE_WITH_REMOTE_OBJECT_HANDLE
+#include "UObject/UObjectMigrationContext.h"
+#endif
+
 #include UE_INLINE_GENERATED_CPP_BY_NAME(PlayerController)
 
 DEFINE_LOG_CATEGORY(LogPlayerController);
@@ -1364,6 +1368,12 @@ void APlayerController::OnActorChannelOpen(FInBunch& InBunch, UNetConnection* Co
 
 	// Attempt to match the player controller to a local viewport (client side)
 	InBunch << NetPlayerIndex;
+
+	if (InBunch.EngineNetVer() >= FEngineNetworkCustomVersion::ClientHandshakeId)
+	{
+		InBunch << ClientHandshakeId;
+	}
+	
 	if (Connection->Driver != NULL && Connection == Connection->Driver->ServerConnection)
 	{
 		const bool bIsParentConnection = 
@@ -1377,18 +1387,21 @@ void APlayerController::OnActorChannelOpen(FInBunch& InBunch, UNetConnection* Co
 		}
 		else
 		{
+			ensure(NetPlayerIndex > 0);
+			
 			int32 ChildIndex = int32(NetPlayerIndex) - 1;
-			if (Connection->Children.IsValidIndex(ChildIndex))
+
+			// NetPlayerIndex represents an index UNetConnection::Children on both the client
+			// and server. This loop will support cases where two player controllers are processed
+			// by the client out of order (e.g. NetPlayerIndex=2 before NetPlayerIndex=1).
+			for (int32 Index = Connection->Children.Num(); Index < (ChildIndex + 1); Index++)
 			{
-				// received a new PlayerController for an already existing child
-				Connection->Children[ChildIndex]->HandleClientPlayer(this, Connection);
+				Connection->Driver->CreateChild(Connection);
 			}
-			else
+
+			if (ensure(Connection->Children.IsValidIndex(ChildIndex)))
 			{
-				// create a split connection on the client
-				UChildConnection* Child = Connection->Driver->CreateChild(Connection); 
-				Child->HandleClientPlayer(this, Connection);
-				UE_LOG(LogNet, Verbose, TEXT("Client received PlayerController=%s. Num child connections=%i."), *GetName(), Connection->Children.Num());
+				Connection->Children[ChildIndex]->HandleClientPlayer(this, Connection);
 			}
 		}
 	}
@@ -1403,6 +1416,7 @@ void APlayerController::OnSerializeNewActor(FOutBunch& OutBunch)
 {
 	// serialize PlayerIndex as part of the initial bunch for PlayerControllers so they can be matched to the correct client-side viewport
 	OutBunch << NetPlayerIndex;
+	OutBunch << ClientHandshakeId;
 }
 
 void APlayerController::OnNetCleanup(UNetConnection* Connection)
@@ -4909,6 +4923,228 @@ void APlayerController::GetLifetimeReplicatedProps(TArray< FLifetimeProperty > &
 	DOREPLIFETIME_WITH_PARAMS_FAST(APlayerController, SpawnLocation, Params);
 }
 
+#if UE_WITH_REMOTE_OBJECT_HANDLE
+void APlayerController::PostMigrate(const FUObjectMigrationContext& MigrationContext)
+{
+	// The ANoPawnPlayerController is server specific and should never be migrated between servers.
+	ensure(!IsA(ANoPawnPlayerController::StaticClass()));
+
+	auto GetNetPlayerIndex = [](UNetConnection* Connection)
+	{
+		if (UChildConnection* ChildConnection = Cast<UChildConnection>(Connection))
+		{
+			if (UNetConnection* ParentConnection = ChildConnection->GetParentConnection())
+			{
+				int32 ConnectionNetPlayerIndex = ParentConnection->Children.Find(ChildConnection);
+				if (ensure(ConnectionNetPlayerIndex != INDEX_NONE))
+				{
+					return ConnectionNetPlayerIndex + 1;
+				}
+			}
+		}
+		else
+		{
+			return 0;
+		}
+	};
+
+	auto GetConnectionWithPlayerController = [](UNetDriver* Driver, APlayerController* PlayerController)
+	{
+		UNetConnection* OutConnection = nullptr;
+		Driver->ForEachClientConnection([&OutConnection, PlayerController](UNetConnection* Connection)
+		{
+			if (Connection->PlayerController == PlayerController)
+			{
+				OutConnection = Connection;
+				return true;
+			}
+
+			return false;
+		});
+
+		return OutConnection;
+	};
+	
+	Super::PostMigrate(MigrationContext);
+	
+	const bool bHasOwnership = UE::RemoteObject::Handle::IsOwned(this);
+	
+	if (MigrationContext.MigrationSide == EObjectMigrationSide::Receive)
+	{
+		EObjectMigrationRecvType MigrationRecvType = MigrationContext.GetObjectMigrationRecvType(this);
+
+		UE_LOG(LogPlayerController, Log, TEXT("[%s] APlayerController::PostMigrate [%s]: %s bHasOwnership=%d MigrationSide=Receive RecvType=%s"), 
+			   *UE::RemoteObject::GetGlobalServerId().ToString(), *FRemoteObjectId(this).ToString(), *GetName(), bHasOwnership, ToString(MigrationRecvType));
+
+		ensure(NetConnection == nullptr);
+
+		const bool bAttachControllerToConnection = true;
+		const bool bRegisterWithWorld = (MigrationRecvType == EObjectMigrationRecvType::AssignedOwnership);
+
+		if (bAttachControllerToConnection)
+		{
+			NetConnection = GetWorld()->GetNetDriver()->FindConnection(CachedConnectionPlayerId);
+			if (ensure(NetConnection))
+			{
+				ensure(NetConnection->PlayerController);
+				ensure(NetConnection->PlayerController->IsA(ANoPawnPlayerController::StaticClass()));
+				const int32 PrevNetPlayerIndex = NetConnection->PlayerController->NetPlayerIndex;
+				const int32 PrevHandshakeTransactionId = NetConnection->PlayerController->GetClientHandshakeId();
+				NetConnection->PlayerController = this;
+				NetConnection->OwningActor = this;
+				Player = NetConnection;
+				NetPlayerIndex = GetNetPlayerIndex(NetConnection);
+				SetClientHandshakeId(NetConnection->GetClientHandshakeId());
+				ensure(PrevNetPlayerIndex == NetPlayerIndex);
+				ensure(PrevHandshakeTransactionId == GetClientHandshakeId());
+
+				UE_LOG(LogPlayerController, Log, TEXT("[%s] APlayerController::PostMigrate [%s]: Binding %s to connection %s"), 
+					   *UE::RemoteObject::GetGlobalServerId().ToString(), *FRemoteObjectId(this).ToString(), 
+					   *GetName(), *NetConnection->GetName());
+			}
+		}
+
+		if (bRegisterWithWorld)
+		{
+			bSendClientRestartRPC = true;
+			GetWorld()->AddController(this);
+		}
+	}
+	else if (MigrationContext.MigrationSide == EObjectMigrationSide::Send)
+	{
+		EObjectMigrationSendType MigrationSendType = MigrationContext.GetObjectMigrationSendType(this);
+
+		UE_LOG(LogPlayerController, Log, TEXT("[%s] APlayerController::PostMigrate [%s]: %s bHasOwnership=%d MigrationSide=Send SendType=%s"), 
+			   *UE::RemoteObject::GetGlobalServerId().ToString(), *FRemoteObjectId(this).ToString(), *GetName(), bHasOwnership, ToString(MigrationSendType));
+
+		const bool bDetachControllerFromConnection = true;
+		const bool bDeregisterFromWorld = (MigrationSendType == EObjectMigrationSendType::ReassignOwnership);
+
+		if (bDetachControllerFromConnection)
+		{
+			UNetConnection* PrevNetConnection = GetConnectionWithPlayerController(GetWorld()->GetNetDriver(), this);
+			if (ensure(PrevNetConnection))
+			{
+				ensure(PrevNetConnection->PlayerController == this);
+				ensure(PrevNetConnection->PlayerController);
+				ensure(PrevNetConnection->OwningActor);
+
+				PrevNetConnection->PlayerController = nullptr;
+				PrevNetConnection->OwningActor = nullptr;
+
+				// Look for a instance of ANoPawnPlayerController with the matching player id to attach to the connection.
+				for (FConstPlayerControllerIterator It = GetWorld()->GetPlayerControllerIterator(); It; ++It)
+				{
+					const TWeakObjectPtr<APlayerController>& Ptr = *It;
+					if (!Ptr.IsRemote())
+					{
+						APlayerController* OtherPlayerController = It->Get();
+						if (OtherPlayerController->IsA(ANoPawnPlayerController::StaticClass()) &&
+							OtherPlayerController->CachedConnectionPlayerId == CachedConnectionPlayerId)
+						{
+							PrevNetConnection->PlayerController = OtherPlayerController;
+							PrevNetConnection->OwningActor = OtherPlayerController;
+						}
+					}
+				}
+
+				if (PrevNetConnection->PlayerController == nullptr)
+				{
+					FActorSpawnParameters SpawnInfo;
+					SpawnInfo.ObjectFlags |= RF_Transient;
+					SpawnInfo.bDeferConstruction = true;
+
+					PrevNetConnection->PlayerController = GetWorld()->SpawnActor<APlayerController>(ANoPawnPlayerController::StaticClass(), FVector(), FRotator(), SpawnInfo);
+					PrevNetConnection->PlayerController->SetRole(ROLE_Authority);
+					PrevNetConnection->PlayerController->SetReplicates(true);
+					PrevNetConnection->PlayerController->SetAutonomousProxy(true);
+					PrevNetConnection->PlayerController->FinishSpawning(FTransform());
+					PrevNetConnection->PlayerController->CachedConnectionPlayerId = CachedConnectionPlayerId;
+					PrevNetConnection->PlayerController->NetPlayerIndex = NetPlayerIndex;
+					PrevNetConnection->PlayerController->SetClientHandshakeId(PrevNetConnection->GetClientHandshakeId());
+					PrevNetConnection->PlayerController->NetConnection = PrevNetConnection;
+
+					PrevNetConnection->OwningActor = PrevNetConnection->PlayerController;
+					
+					UE_LOG(LogPlayerController, Log, TEXT("[%s] APlayerController::PostMigrate [%s]: Replacing %s with new player controller %s on connection %s"), 
+						   *UE::RemoteObject::GetGlobalServerId().ToString(), *FRemoteObjectId(this).ToString(), 
+						   *GetName(), *PrevNetConnection->PlayerController->GetName(), *PrevNetConnection->GetName());
+
+					GetWorld()->AddController(PrevNetConnection->PlayerController);
+				}
+				else
+				{
+					UE_LOG(LogPlayerController, Log, TEXT("[%s] APlayerController::PostMigrate [%s]: Replacing %s with existing player controller %s on connection %s"),
+						   *UE::RemoteObject::GetGlobalServerId().ToString(), *FRemoteObjectId(this).ToString(), 
+						   *GetName(), *PrevNetConnection->PlayerController->GetName(), *PrevNetConnection->GetName());
+				}
+
+				ensure(PrevNetConnection->PlayerController);
+				ensure(PrevNetConnection->OwningActor);
+			}
+		}
+
+		if (bDeregisterFromWorld)
+		{
+			UE_LOG(LogPlayerController, Log, TEXT("[%s] APlayerController::PostMigrate [%s]: Removing %s from the world"), 
+				   *UE::RemoteObject::GetGlobalServerId().ToString(), *FRemoteObjectId(this).ToString(), *GetName());
+
+			GetWorld()->RemoveController(this);
+		}
+	}
+}
+#endif
+
+void APlayerController::Serialize(FArchive& Ar)
+{
+#if UE_WITH_REMOTE_OBJECT_HANDLE
+	if (Ar.IsMigratingRemoteObjects())
+	{
+		if (Ar.IsLoading() || Ar.IsSaving())
+		{
+			if (NetConnection)
+			{
+				CachedConnectionPlayerId = NetConnection->PlayerId;
+			}
+
+			Ar << CachedConnectionPlayerId;
+
+			// The these properties must be nulled out here because they are server specific. This must happen here
+			// so they are serialized as null in Super::Serialize().
+			NetConnection = nullptr;
+			Player = nullptr;
+		}
+	}
+		
+	Super::Serialize(Ar);
+#else
+	Super::Serialize(Ar);
+#endif
+}
+
+#if UE_WITH_REMOTE_OBJECT_HANDLE
+void APlayerController::Tick(float DeltaTime)
+{
+	Super::Tick(DeltaTime);
+
+	if (bSendClientRestartRPC)
+	{
+		if (!GetPawn().IsRemote())
+		{
+			UE_LOG(LogPlayerController, Log, TEXT("[%s] APlayerController::Tick [%s]: Repossessing pawn %s after migration."), 
+				   *UE::RemoteObject::GetGlobalServerId().ToString(),
+				   *FRemoteObjectId(this).ToString(),
+				   *GetPawn()->GetName());
+
+			GetPawn()->OnDestroyed.Clear();
+			Possess(GetPawn());
+			
+			bSendClientRestartRPC = false;
+		}
+	}
+}
+#endif
+
 void APlayerController::SetTargetViewRotation(const FRotator& InRotation)
 {
 	if (UE::Gameplay::CVars::bIsPlayerControllerPushBased)
@@ -6224,6 +6460,16 @@ FAsyncPhysicsTimestamp APlayerController::GetPhysicsTimestamp(float DeltaSeconds
 	return Timestamp;
 }
 
+uint32 APlayerController::GetClientHandshakeId() const
+{
+	return ClientHandshakeId;
+}
+
+void APlayerController::SetClientHandshakeId(uint32 InClientHandshakeId)
+{
+	ClientHandshakeId = InClientHandshakeId;
+}
+
 void APlayerController::UpdateServerAsyncPhysicsTickOffset()
 {
 	if (NetworkPhysicsCvars::TickOffsetUpdateIntervalTime <= 0 || TickOffsetSyncCountdown < (NetworkPhysicsCvars::TickOffsetUpdateIntervalTime * 0.001f))
diff --git a/Engine/Source/Runtime/Engine/Private/World.cpp b/Engine/Source/Runtime/Engine/Private/World.cpp
--- a/Engine/Source/Runtime/Engine/Private/World.cpp
+++ b/Engine/Source/Runtime/Engine/Private/World.cpp
@@ -234,6 +234,23 @@ namespace UE::Private::World
 
 		FStreamingLevelsToConsider& StreamingLevelsToConsider;
 	};
+
+	void ParseAndSetClientHandshakeId(FURL& InUrl, UNetConnection* Connection)
+	{
+		if (const TCHAR* HandshakeIdStr = InUrl.GetOption(TEXT("HandshakeId="), nullptr))
+		{
+			TCHAR* EndStr = nullptr;
+			uint32 HandshakeId = FCString::Strtoi(HandshakeIdStr, &EndStr, 10);
+			if (HandshakeIdStr != EndStr && HandshakeId )
+			{
+				Connection->SetClientHandshakeId(HandshakeId);
+			}
+			else
+			{
+				UE_LOG(LogNet, Warning, TEXT("Connection %s received an invalid client handshake id."), *Connection->GetName());
+			}
+		}
+	}
 }
 
 namespace UE::Gameplay::CVars
@@ -7068,7 +7085,9 @@ void UWorld::NotifyControlMessage(UNetConnection* Connection, uint8 MessageType,
 						break;
 					}
 
-					Connection->PlayerController = SpawnPlayActor( Connection, ROLE_AutonomousProxy, InURL, Connection->PlayerId, ErrorMsg );
+					UE::Private::World::ParseAndSetClientHandshakeId(InURL, Connection);
+
+					Connection->PlayerController = SpawnPlayActor( Connection, ROLE_AutonomousProxy, InURL, Connection->PlayerId, ErrorMsg, 0 );
 					if (Connection->PlayerController == NULL)
 					{
 						// Failed to connect.
@@ -7180,6 +7199,17 @@ void UWorld::NotifyControlMessage(UNetConnection* Connection, uint8 MessageType,
 			{
 				if (NetDriver->HasSupportForNoPawnConnection())
 				{
+					FURL InURL(NULL, *Connection->RequestURL, TRAVEL_Absolute);
+
+					if (!InURL.Valid)
+					{
+						UE_LOG(LogNet, Error, TEXT("NMT_JoinNoPawn: Invalid URL %s"), *Connection->RequestURL);
+						Bunch.SetError();
+						break;
+					}
+
+					UE::Private::World::ParseAndSetClientHandshakeId(InURL, Connection);
+
 					FActorSpawnParameters SpawnInfo;
 					SpawnInfo.ObjectFlags |= RF_Transient;
 					SpawnInfo.bDeferConstruction = true;
@@ -7191,12 +7221,12 @@ void UWorld::NotifyControlMessage(UNetConnection* Connection, uint8 MessageType,
 					Connection->PlayerController->FinishSpawning(FTransform());
 					Connection->PlayerController->NetPlayerIndex = 0;
 					Connection->PlayerController->NetConnection = Connection;
+					Connection->PlayerController->SetClientHandshakeId(Connection->GetClientHandshakeId());
 
 					Connection->OwningActor = Connection->PlayerController;
 				}
 				else
 				{
-					UE_LOG(LogNet, Log, TEXT("Server doesn't support NMT_JoinNoPawn (enable with UNetDriver::SetSupportForNoPawnConnection() or net.EnableSupportForNoPawnConnection)."));
 					FString Msg(TEXT("Server doesn't support NMT_JoinNoPawn."));
 					Connection->SendCloseReason(ENetCloseResult::JoinFailure);
 					FNetControlMessage<NMT_Failure>::Send(Connection, Msg);
@@ -7210,15 +7240,27 @@ void UWorld::NotifyControlMessage(UNetConnection* Connection, uint8 MessageType,
 			{
 				if (NetDriver->HasSupportForNoPawnConnection())
 				{
+					FString SplitRequestURL;
 					FUniqueNetIdRepl UniqueIdRepl;
-					if (FNetControlMessage<NMT_JoinNoPawnSplit>::Receive(Bunch, UniqueIdRepl))
+					if (FNetControlMessage<NMT_JoinNoPawnSplit>::Receive(Bunch, SplitRequestURL, UniqueIdRepl))
 					{
+						FURL InURL(NULL, *SplitRequestURL, TRAVEL_Absolute);
+
+						if (!InURL.Valid)
+						{
+							UE_LOG(LogNet, Error, TEXT("NMT_JoinNoPawnSplit: Invalid URL %s"), *SplitRequestURL);
+							Bunch.SetError();
+							break;
+						}
+
 						FActorSpawnParameters SpawnInfo;
 						SpawnInfo.ObjectFlags |= RF_Transient;
 						SpawnInfo.bDeferConstruction = true;
 					
 						UChildConnection* ChildConnection = NetDriver->CreateChild(Connection);
 						ChildConnection->PlayerId = UniqueIdRepl;
+						
+						UE::Private::World::ParseAndSetClientHandshakeId(InURL, ChildConnection);
 
 						ChildConnection->PlayerController = GetWorld()->SpawnActor<APlayerController>(ANoPawnPlayerController::StaticClass(), FVector(), FRotator(), SpawnInfo);
 						ChildConnection->PlayerController->SetRole(ROLE_Authority);
@@ -7227,6 +7269,7 @@ void UWorld::NotifyControlMessage(UNetConnection* Connection, uint8 MessageType,
 						ChildConnection->PlayerController->FinishSpawning(FTransform());
 						ChildConnection->PlayerController->NetPlayerIndex = Connection->Children.Num();
 						ChildConnection->PlayerController->NetConnection = ChildConnection;
+						ChildConnection->PlayerController->SetClientHandshakeId(ChildConnection->GetClientHandshakeId());
 
 						ChildConnection->OwningActor = ChildConnection->PlayerController;
 					}
@@ -7349,6 +7392,8 @@ void UWorld::PreLoginCompleteSplit(const FString& ErrorMsg, TWeakObjectPtr<UNetC
 	// create URL from string
 	FURL JoinSplitURL(NULL, *SplitRequestURL, TRAVEL_Absolute);
 
+	UE::Private::World::ParseAndSetClientHandshakeId(JoinSplitURL, ChildConn);
+
 	UE_LOG(LogNet, Log, TEXT("JOINSPLIT: Join request: URL=%s"), *JoinSplitURL.ToString());
 	FString SpawnErrorMsg;
 	APlayerController* PC = SpawnPlayActor(ChildConn, ROLE_AutonomousProxy, JoinSplitURL, ChildConn->PlayerId, SpawnErrorMsg, uint8(Connection->Children.Num()));
diff --git a/Engine/Source/Runtime/Engine/Public/Net/DataChannel.h b/Engine/Source/Runtime/Engine/Public/Net/DataChannel.h
--- a/Engine/Source/Runtime/Engine/Public/Net/DataChannel.h
+++ b/Engine/Source/Runtime/Engine/Public/Net/DataChannel.h
@@ -228,4 +228,4 @@ DEFINE_CONTROL_CHANNEL_MESSAGE(IrisNetRefHandleError, 30, UE::Net::ENetRefHandle
 // by default and must be explicitly enabled on the receiving server using the cvar net.EnableSupportForNoPawnConnection or calling
 // UNetDriver::SetSupportForNoPawnConnection() directly.
 DEFINE_CONTROL_CHANNEL_MESSAGE(JoinNoPawn, 31);
-DEFINE_CONTROL_CHANNEL_MESSAGE(JoinNoPawnSplit, 32, FUniqueNetIdRepl);
\ No newline at end of file
+DEFINE_CONTROL_CHANNEL_MESSAGE(JoinNoPawnSplit, 32, FString, FUniqueNetIdRepl);
\ No newline at end of file
diff --git a/Engine/Source/Runtime/GameplayDebugger/Private/GameplayDebuggerCategoryReplicator.cpp b/Engine/Source/Runtime/GameplayDebugger/Private/GameplayDebuggerCategoryReplicator.cpp
--- a/Engine/Source/Runtime/GameplayDebugger/Private/GameplayDebuggerCategoryReplicator.cpp
+++ b/Engine/Source/Runtime/GameplayDebugger/Private/GameplayDebuggerCategoryReplicator.cpp
@@ -676,7 +676,7 @@ void AGameplayDebuggerCategoryReplicator::OnExtensionsChanged()
 
 UNetConnection* AGameplayDebuggerCategoryReplicator::GetNetConnection() const
 {
-	return IsValid(OwnerPC) ? OwnerPC->GetNetConnection() : nullptr;
+	return (!OwnerPC.IsRemote() && IsValid(OwnerPC)) ? OwnerPC->GetNetConnection() : nullptr;
 }
 
 bool AGameplayDebuggerCategoryReplicator::IsNetRelevantFor(const AActor* RealViewer, const AActor* ViewTarget, const FVector& SrcLocation) const
diff --git a/Engine/Source/Runtime/GameplayDebugger/Private/GameplayDebuggerLocalController.cpp b/Engine/Source/Runtime/GameplayDebugger/Private/GameplayDebuggerLocalController.cpp
--- a/Engine/Source/Runtime/GameplayDebugger/Private/GameplayDebuggerLocalController.cpp
+++ b/Engine/Source/Runtime/GameplayDebugger/Private/GameplayDebuggerLocalController.cpp
@@ -915,7 +915,7 @@ void UGameplayDebuggerLocalController::OnSelectedObject(UObject* Object)
 	}
 
 	AController* SelectedController = Cast<AController>(Object);
-	APawn* SelectedPawn = SelectedController ? SelectedController->GetPawn() : Cast<APawn>(Object);
+	APawn* SelectedPawn = SelectedController ? SelectedController->GetPawn().Get() : Cast<APawn>(Object);
 	if (CachedReplicator && SelectedPawn && SelectedPawn->IsSelected())
 	{
 		CachedReplicator->SetDebugActor(SelectedPawn, false);
diff --git a/Engine/Source/Runtime/GameplayDebugger/Private/GameplayDebuggerPlayerManager.cpp b/Engine/Source/Runtime/GameplayDebugger/Private/GameplayDebuggerPlayerManager.cpp
--- a/Engine/Source/Runtime/GameplayDebugger/Private/GameplayDebuggerPlayerManager.cpp
+++ b/Engine/Source/Runtime/GameplayDebugger/Private/GameplayDebuggerPlayerManager.cpp
@@ -159,6 +159,11 @@ void AGameplayDebuggerPlayerManager::UpdateAuthReplicators()
 	{
 		FGameplayDebuggerPlayerData& TestData = PlayerData[Idx];
 
+		if (IsValid(TestData.Replicator) && TestData.Replicator->GetReplicationOwner().IsRemote())
+		{
+			continue;
+		}
+
 		if (!IsValid(TestData.Replicator) || !IsValid(TestData.Replicator->GetReplicationOwner()))
 		{
 			if (IsValid(TestData.Replicator))
diff --git a/Engine/Source/Runtime/GameplayDebugger/Public/GameplayDebuggerCategoryReplicator.h b/Engine/Source/Runtime/GameplayDebugger/Public/GameplayDebuggerCategoryReplicator.h
--- a/Engine/Source/Runtime/GameplayDebugger/Public/GameplayDebuggerCategoryReplicator.h
+++ b/Engine/Source/Runtime/GameplayDebugger/Public/GameplayDebuggerCategoryReplicator.h
@@ -186,7 +186,7 @@ public:
 	const FGameplayDebuggerVisLogSync& GetVisLogSyncData() const { return VisLogSync; }
 
 	/** get player controller owning this replicator */
-	APlayerController* GetReplicationOwner() const { return OwnerPC; }
+	TObjectPtr<APlayerController> GetReplicationOwner() const { return OwnerPC; }
 
 	/** get replicator state */
 	bool IsEnabled() const { return bIsEnabled; }
