Cleanup dependency NetDrivers created by the MultiServer Proxy.

diff --git a/Engine/Plugins/Runtime/MultiServerReplication/Source/MultiServerReplication/Private/MultiServerProxy.cpp b/Engine/Plugins/Runtime/MultiServerReplication/Source/MultiServerReplication/Private/MultiServerProxy.cpp
--- a/Engine/Plugins/Runtime/MultiServerReplication/Source/MultiServerReplication/Private/MultiServerProxy.cpp
+++ b/Engine/Plugins/Runtime/MultiServerReplication/Source/MultiServerReplication/Private/MultiServerProxy.cpp
@@ -130,11 +130,15 @@ void UProxyListenerNotify::ConnectToGameServer(UNetConnection* ProxyConnection,
 	// connections from proxy connections to that server.
 	if (GameServerConnectionState->NetDriver == nullptr)
 	{
-		static int32 GameServerDriverIndex = 0;
-		FString Name = FString::Printf(TEXT("ProxyToGameServer-%d"), GameServerDriverIndex++);
-		GEngine->CreateNamedNetDriver(ProxyNetDriver->GetWorld(), *Name, "ProxyBackendNetDriver");
+		// Acts as a unique identifier for dependency NetDrivers.
+		static int32 GameServerDriverId = 0;
 
-		GameServerConnectionState->NetDriver = GEngine->FindNamedNetDriver(ProxyNetDriver->GetWorld(), *Name);
+		GameServerConnectionState->World = ProxyNetDriver->GetWorld();
+		GameServerConnectionState->NetDriverName = FName(*FString::Printf(TEXT("ProxyToGameServer-%d"), GameServerDriverId++));
+
+		GEngine->CreateNamedNetDriver(ProxyNetDriver->GetWorld(), GameServerConnectionState->NetDriverName, "ProxyBackendNetDriver");
+		GameServerConnectionState->NetDriver = GEngine->FindNamedNetDriver(GameServerConnectionState->World, GameServerConnectionState->NetDriverName);
+		
 		GameServerConnectionState->NetDriver->SetWorld(ProxyNetDriver->GetWorld());
 		Cast<UProxyBackendNetDriver>(GameServerConnectionState->NetDriver)->SetProxyNetDriver(ProxyNetDriver);
 
@@ -751,6 +755,19 @@ bool UProxyNetDriver::CanDowngradeActorRole(UNetConnection* ProxyConnection, AAc
 	return false;
 }
 
+void UProxyNetDriver::Shutdown()
+{
+	for (FGameServerConnectionState& ConnectionState : GameServerConnections)
+	{
+		GEngine->DestroyNamedNetDriver(ConnectionState.World, ConnectionState.NetDriverName);
+		ConnectionState.GameServerNotify->NetDriver = nullptr;
+	}
+
+	GameServerConnections.Reset();
+
+	Super::Shutdown();
+}
+
 int32 UProxyNetDriver::GetNextGameServerClientId()
 {
 	return NextGameServerClientId++;
diff --git a/Engine/Plugins/Runtime/MultiServerReplication/Source/MultiServerReplication/Public/MultiServerProxy.h b/Engine/Plugins/Runtime/MultiServerReplication/Source/MultiServerReplication/Public/MultiServerProxy.h
--- a/Engine/Plugins/Runtime/MultiServerReplication/Source/MultiServerReplication/Public/MultiServerProxy.h
+++ b/Engine/Plugins/Runtime/MultiServerReplication/Source/MultiServerReplication/Public/MultiServerProxy.h
@@ -41,8 +41,13 @@ struct FGameServerConnectionState
 
 	int32 Port;
 
+	UPROPERTY()
+	TObjectPtr<class UWorld> World;
+
 	UPROPERTY()
 	TObjectPtr<class UNetDriver> NetDriver;
+	
+	FName NetDriverName;
 
 	UPROPERTY()
 	TArray<TObjectPtr<class ULocalPlayer>> Players;
@@ -276,6 +281,7 @@ public:
 	virtual void InternalProcessRemoteFunction(class AActor* Actor, class UObject* SubObject, class UNetConnection* Connection, class UFunction* Function, void* Parms, FOutParmRec* OutParms, FFrame* Stack, bool bIsServer) override;
 	virtual int32 ServerReplicateActors(float DeltaSeconds) override;
 	virtual bool CanDowngradeActorRole(UNetConnection* ProxyConnection, AActor* Actor) const override;
+	virtual void Shutdown() override;
 	
 	/** Get the next identifier for outgoing connections to game servers. */
 	int32 GetNextGameServerClientId();
