Disable calls to AActor::PreReplication() on the multiserver proxy.

diff --git a/Engine/Plugins/Runtime/MultiServerReplication/Source/MultiServerReplication/Private/MultiServerProxy.cpp b/Engine/Plugins/Runtime/MultiServerReplication/Source/MultiServerReplication/Private/MultiServerProxy.cpp
--- a/Engine/Plugins/Runtime/MultiServerReplication/Source/MultiServerReplication/Private/MultiServerProxy.cpp
+++ b/Engine/Plugins/Runtime/MultiServerReplication/Source/MultiServerReplication/Private/MultiServerProxy.cpp
@@ -529,13 +529,7 @@ bool UProxyNetDriver::InitBase(bool bInitAsClient, FNetworkNotify* InNotify, con
 	ProxyListenerNotify = NewObject<UProxyListenerNotify>(GEngine, UProxyListenerNotify::StaticClass());
 	ProxyListenerNotify->SetProxyNetDriver(this);
 
-	bool bSuccess = Super::InitBase(bInitAsClient, ProxyListenerNotify, URL, bReuseAddressAndPort, Error);
-
-	if (bSuccess)
-	{
-		// Don't allow any RPCs received from clients to be executed, only forwarded to the game server.
-		EnableExecuteRPCFunctions(false);
-	}
+	const bool bSuccess = Super::InitBase(bInitAsClient, ProxyListenerNotify, URL, bReuseAddressAndPort, Error);
 
 	FString GameServerAddresses;
 	if (FParse::Value(FCommandLine::Get(), TEXT("-ProxyGameServers="), GameServerAddresses, false))
@@ -566,11 +560,7 @@ bool UProxyNetDriver::InitBase(bool bInitAsClient, FNetworkNotify* InNotify, con
 
 	SetReplicateTransactionally(false);
 	
-#if UE_SUPPORT_FOR_ACTOR_TICK_DISABLE
-	// The proxy should only be replicating the exact state from the servers and passing it on
-	// to clients and not tick actors or call any user actor callbacks.
-	GetWorld()->EnableActorTickAndUserCallbacks(false);
-#endif
+	DisableActorLogicAndGameCode();
 
 	return bSuccess;
 }
@@ -778,6 +768,19 @@ uint32 UProxyNetDriver::GetNextClientHandshakeId()
 	return NextClientHandshakeId++;
 }
 
+void UProxyNetDriver::DisableActorLogicAndGameCode()
+{
+	// The proxy should only be replicating the exact state from the servers and passing it on
+	// to clients and not tick actors or call any user actor callbacks.
+
+	EnableExecuteRPCFunctions(false);
+	EnablePreReplication(false);
+
+#if UE_SUPPORT_FOR_ACTOR_TICK_DISABLE
+	GetWorld()->EnableActorTickAndUserCallbacks(false);
+#endif
+}
+
 bool UProxyNetDriver::IsProxySpawned(AActor* Actor) const
 {
 	return ProxySpawnedActors.Contains(Actor);
diff --git a/Engine/Plugins/Runtime/MultiServerReplication/Source/MultiServerReplication/Public/MultiServerProxy.h b/Engine/Plugins/Runtime/MultiServerReplication/Source/MultiServerReplication/Public/MultiServerProxy.h
--- a/Engine/Plugins/Runtime/MultiServerReplication/Source/MultiServerReplication/Public/MultiServerProxy.h
+++ b/Engine/Plugins/Runtime/MultiServerReplication/Source/MultiServerReplication/Public/MultiServerProxy.h
@@ -303,6 +303,9 @@ private:
 	
 	friend class UProxyBackendNetDriver;
 
+	/** Set all of the configuration options that disable executing actor functionality or game specific code. */
+	void DisableActorLogicAndGameCode();
+
 	/** Return true if an actor is explicitly spawned by the proxy. */
 	bool IsProxySpawned(AActor* Actor) const;
 	
diff --git a/Engine/Source/Runtime/Engine/Classes/Engine/NetDriver.h b/Engine/Source/Runtime/Engine/Classes/Engine/NetDriver.h
--- a/Engine/Source/Runtime/Engine/Classes/Engine/NetDriver.h
+++ b/Engine/Source/Runtime/Engine/Classes/Engine/NetDriver.h
@@ -1476,6 +1476,10 @@ public:
 	ENGINE_API void EnableExecuteRPCFunctions(bool bEnable);
 	ENGINE_API bool IsExecuteRPCFunctionsEnabled() const;
 
+	/** (Generic Replication Only) Enable/disable calls to AActor::PreReplication(). */
+	ENGINE_API void EnablePreReplication(bool bEnable);
+	ENGINE_API bool IsPreReplicationEnabled() const;
+
 	/** (Generic Replicaton Only) Treat actors with ROLE_AutonomousProxy as ROLE_SimulatedProxy when replicating state (applies to all client connections). */
 	ENGINE_API void ReplicateAutonomousAsSimulated(bool bEnabled);
 	ENGINE_API bool IsReplicatingAutonomousAsSimulated() const;
@@ -1486,6 +1490,7 @@ public:
 private:
 
 	bool bExecuteRPCFunctions = true;
+	bool bCallPreReplication = true;
 	bool bReplicatingAutonomousAsSimulated = false;
 
 public:
diff --git a/Engine/Source/Runtime/Engine/Private/NetDriver.cpp b/Engine/Source/Runtime/Engine/Private/NetDriver.cpp
--- a/Engine/Source/Runtime/Engine/Private/NetDriver.cpp
+++ b/Engine/Source/Runtime/Engine/Private/NetDriver.cpp
@@ -2123,6 +2123,16 @@ bool UNetDriver::IsExecuteRPCFunctionsEnabled() const
 	return bExecuteRPCFunctions;
 }
 
+void UNetDriver::EnablePreReplication(bool bEnable)
+{
+	bCallPreReplication = bEnable;
+}
+
+bool UNetDriver::IsPreReplicationEnabled() const
+{
+	return bCallPreReplication;
+}
+
 void UNetDriver::ReplicateAutonomousAsSimulated(bool bEnabled)
 {
 	bReplicatingAutonomousAsSimulated = bEnabled;
@@ -2866,7 +2876,12 @@ void UNetDriver::ProcessRemoteFunctionForChannelPrivate(
 		if (!EnumHasAnyFlags(RemoteFunctionFlags, EProcessRemoteFunctionFlags::ReplicatedActor))
 		{
 			ReplicationFrame++;
-			Ch->GetActor()->CallPreReplication(this);
+
+			if (bCallPreReplication)
+			{
+				Ch->GetActor()->CallPreReplication(this);
+			}
+
 			RemoteFunctionFlags |= EProcessRemoteFunctionFlags::ReplicatedActor;
 		}
 
@@ -5287,7 +5302,10 @@ void UNetDriver::ServerReplicateActors_BuildConsiderList( TArray<FNetworkObjectI
 		OutConsiderList.Add( ActorInfo );
 
 		// Call PreReplication on all actors that will be considered
-		Actor->CallPreReplication( this );
+		if (bCallPreReplication)
+		{
+			Actor->CallPreReplication(this);
+		}
 	}
 
 	for ( AActor* Actor : ActorsToRemove )
