Fix a use after free crash in the multi-server proxy.

diff --git a/Engine/Plugins/Runtime/MultiServerReplication/Source/MultiServerReplication/Private/MultiServerProxy.cpp b/Engine/Plugins/Runtime/MultiServerReplication/Source/MultiServerReplication/Private/MultiServerProxy.cpp
--- a/Engine/Plugins/Runtime/MultiServerReplication/Source/MultiServerReplication/Private/MultiServerProxy.cpp
+++ b/Engine/Plugins/Runtime/MultiServerReplication/Source/MultiServerReplication/Private/MultiServerProxy.cpp
@@ -798,22 +798,25 @@ void UProxyNetDriver::GameServerAssignPlayerController(UNetConnection* ChildGame
 	GameServerPlayerController->SetRole(ROLE_AutonomousProxy);
 
 	const uint32 ClientHandshakeId = GameServerPlayerController->GetClientHandshakeId();
-	FMultiServerProxyInternalConnectionRoute* Route = ClientHandshakeIdToInternalRoute.Find(ClientHandshakeId);
-	if (Route == nullptr)
+	FMultiServerProxyInternalConnectionRoute* RoutePtr = ClientHandshakeIdToInternalRoute.Find(ClientHandshakeId);
+	if (RoutePtr == nullptr)
 	{
 		UE_LOG(LogNetProxy, Error, TEXT("Failed to find a mapping for game server connection request (client_handshake_id=%u)"), ClientHandshakeId);
 		return;
 	}
+
+	// Cache this value locally incase ClientHandshakeIdToInternalRoute changes and invalidates the pointer.
+	FMultiServerProxyInternalConnectionRoute Route = *RoutePtr;
 	
 	UE_LOG(LogNetProxy, Log, TEXT("Found internal route (client_handshake_id=%u player=%s proxy_connection=%s:%s parent_game_server_connection=%s:%s)"), 
 		   ClientHandshakeId,
-		   *Route->Player->GetName(),
-		   *Route->ProxyConnection->GetDriver()->GetName(),
-		   *Route->ProxyConnection->GetName(),
-		   *Route->ParentGameServerConnection->GetDriver()->GetName(),
-		   *Route->ParentGameServerConnection->GetName());
+		   *Route.Player->GetName(),
+		   *Route.ProxyConnection->GetDriver()->GetName(),
+		   *Route.ProxyConnection->GetName(),
+		   *Route.ParentGameServerConnection->GetDriver()->GetName(),
+		   *Route.ParentGameServerConnection->GetName());
 
-	ULocalPlayer* Player = Route->Player;
+	ULocalPlayer* Player = Route.Player;
 		
 	// The player should only reference a proxy player controller if this game server connection
 	// is the primary game server for the proxy connection (decided later in this function).
@@ -826,7 +829,7 @@ void UProxyNetDriver::GameServerAssignPlayerController(UNetConnection* ChildGame
 	AttachPlayerController(GameServerPlayerController, ChildGameServerConnection, Player);
 
 	// Find the proxy connection that established this connection to the game server.
-	TObjectPtr<UNetConnection> ProxyConnection = Route->ProxyConnection;
+	TObjectPtr<UNetConnection> ProxyConnection = Route.ProxyConnection;
 
 	if (ensure(ProxyConnectionState.Contains(ProxyConnection)))
 	{
@@ -845,14 +848,16 @@ void UProxyNetDriver::GameServerAssignPlayerController(UNetConnection* ChildGame
 
 			case EProxyConnectionState::ConnectingPrimary:
 				{
-					if (int32* PlayerPrimaryGameServerIdx = ProxyPlayerToPrimaryGameServer.Find(Player))
+					if (int32* PlayerPrimaryGameServerIdxPtr = ProxyPlayerToPrimaryGameServer.Find(Player))
 					{
+						int32 PlayerPrimaryGameServerIdx = *PlayerPrimaryGameServerIdxPtr;
+
 						ensure(ProxyConnection->PlayerController == nullptr);
 
-						if (ensure(*PlayerPrimaryGameServerIdx < GameServerConnections.Num()))
+						if (ensure(PlayerPrimaryGameServerIdx < GameServerConnections.Num()))
 						{
-							FGameServerConnectionState& PlayerPrimaryGameServer = GameServerConnections[*PlayerPrimaryGameServerIdx];
-							if (PlayerPrimaryGameServer.NetDriver->ServerConnection == Route->ParentGameServerConnection)
+							FGameServerConnectionState& PlayerPrimaryGameServer = GameServerConnections[PlayerPrimaryGameServerIdx];
+							if (PlayerPrimaryGameServer.NetDriver->ServerConnection == Route.ParentGameServerConnection)
 							{
 								UPlayer* PreviousPlayer = GameServerPlayerController->Player;
 			
@@ -903,9 +908,11 @@ void UProxyNetDriver::GameServerAssignPlayerController(UNetConnection* ChildGame
 						{
 							if (ProxyPlayerToPrimaryGameServerPlayerController.Contains(ExistingPrimaryPlayer))
 							{
-								int32* PreviousPrimaryGameServerIdx = ProxyPlayerToPrimaryGameServer.Find(ExistingPrimaryPlayer);
-								if (ensure(PreviousPrimaryGameServerIdx))
+								int32* PreviousPrimaryGameServerIdxPtr = ProxyPlayerToPrimaryGameServer.Find(ExistingPrimaryPlayer);
+								if (ensure(PreviousPrimaryGameServerIdxPtr))
 								{
+									int32 PreviousPrimaryGameServerIdx = *PreviousPrimaryGameServerIdxPtr;
+
 									if (Player->GetCachedUniqueNetId() == ExistingPrimaryPlayer->GetCachedUniqueNetId())
 									{
 										int32 NewPrimaryGameServerIdx = GetGameServerWithPlayer(Player);
@@ -915,7 +922,7 @@ void UProxyNetDriver::GameServerAssignPlayerController(UNetConnection* ChildGame
 											*ProxyConnection->GetName(),
 											*ExistingPrimaryPlayer->GetName(),
 											*Player->GetName(),
-											*PreviousPrimaryGameServerIdx,
+											PreviousPrimaryGameServerIdx,
 											NewPrimaryGameServerIdx);
 					
 										ProxyPlayerToPrimaryGameServerPlayerController.Remove(ExistingPrimaryPlayer);
@@ -941,15 +948,15 @@ void UProxyNetDriver::GameServerAssignPlayerController(UNetConnection* ChildGame
 			
 	// Update the internal connection mapping to use the new child game server connection (if it's different from the parent connection).
 	ClientHandshakeIdToInternalRoute.Remove(NextClientHandshakeId);
-	ClientHandshakeIdToInternalRoute.Add(NextClientHandshakeId, *Route);
+	ClientHandshakeIdToInternalRoute.Add(NextClientHandshakeId, Route);
 
 	// If this was a parent connection send any join requests for any pending multiplexed connections.
-	if (ChildGameServerConnection == Route->ParentGameServerConnection)
+	if (ChildGameServerConnection == Route.ParentGameServerConnection)
 	{
 		FGameServerConnectionState* CurrGameServerConnectionState = nullptr;
 		for (FGameServerConnectionState& GameServerConnectionState : GameServerConnections)
 		{
-			if (GameServerConnectionState.NetDriver->ServerConnection == Route->ParentGameServerConnection)
+			if (GameServerConnectionState.NetDriver->ServerConnection == Route.ParentGameServerConnection)
 			{
 				CurrGameServerConnectionState = &GameServerConnectionState;
 				break;
