Preventing crashes in the logging functions during disconnection from the multi-server proxy (NetDriver was nullptr).

#rnx

diff --git a/Engine/Plugins/Runtime/MultiServerReplication/Source/MultiServerReplication/Private/MultiServerProxy.cpp b/Engine/Plugins/Runtime/MultiServerReplication/Source/MultiServerReplication/Private/MultiServerProxy.cpp
--- a/Engine/Plugins/Runtime/MultiServerReplication/Source/MultiServerReplication/Private/MultiServerProxy.cpp
+++ b/Engine/Plugins/Runtime/MultiServerReplication/Source/MultiServerReplication/Private/MultiServerProxy.cpp
@@ -810,11 +810,11 @@ void UProxyNetDriver::GameServerAssignPlayerController(UNetConnection* ChildGame
 	
 	UE_LOG(LogNetProxy, Log, TEXT("Found internal route (client_handshake_id=%u player=%s proxy_connection=%s:%s parent_game_server_connection=%s:%s)"), 
 		   ClientHandshakeId,
-		   *Route.Player->GetName(),
-		   *Route.ProxyConnection->GetDriver()->GetName(),
-		   *Route.ProxyConnection->GetName(),
-		   *Route.ParentGameServerConnection->GetDriver()->GetName(),
-		   *Route.ParentGameServerConnection->GetName());
+		   *GetNameSafe(Route.Player),
+			Route.ProxyConnection ? *GetNameSafe(Route.ProxyConnection->GetDriver()) : TEXT("None"),
+		   *GetNameSafe(Route.ProxyConnection),
+			Route.ParentGameServerConnection ? *GetNameSafe(Route.ParentGameServerConnection->GetDriver()) : TEXT("None"),
+		   *GetNameSafe(Route.ParentGameServerConnection));
 
 	ULocalPlayer* Player = Route.Player;
 		
@@ -836,9 +836,9 @@ void UProxyNetDriver::GameServerAssignPlayerController(UNetConnection* ChildGame
 		EProxyConnectionState State = ProxyConnectionState[ProxyConnection];
 
 		UE_LOG(LogNetProxy, Log, TEXT("Attempting to assign game server player controller %s to proxy connection %s:%s (state=%s)"),
-			*GameServerPlayerController->GetName(),
-			*ProxyConnection->GetDriver()->GetName(),
-			*ProxyConnection->GetName(),
+			*GetNameSafe(GameServerPlayerController),
+			ProxyConnection ? *GetNameSafe(ProxyConnection->GetDriver()) : TEXT("None"),
+			*GetNameSafe(ProxyConnection),
 			*LexToString(State));
 
 		switch (State)
