Update the ticking structure of MultiServerReplication netdrivers to support ticking them multiple times per frame while Iris replication is enabled.
Also fix up AActor::PostMigrate to support Iris replication.

diff --git a/Engine/Plugins/Runtime/MultiServerReplication/Source/MultiServerReplication/Private/MultiServerBeaconClient.cpp b/Engine/Plugins/Runtime/MultiServerReplication/Source/MultiServerReplication/Private/MultiServerBeaconClient.cpp
--- a/Engine/Plugins/Runtime/MultiServerReplication/Source/MultiServerReplication/Private/MultiServerBeaconClient.cpp
+++ b/Engine/Plugins/Runtime/MultiServerReplication/Source/MultiServerReplication/Private/MultiServerBeaconClient.cpp
@@ -2,6 +2,7 @@
 
 #include "MultiServerBeaconClient.h"
 #include "MultiServerBeaconHostObject.h"
+#include "MultiServerNetDriver.h"
 #include "MultiServerNode.h"
 #include "Engine/Engine.h"
 #include "Engine/Level.h"
@@ -132,6 +133,8 @@ bool AMultiServerBeaconClient::InitBase()
 {
 	if (Super::InitBase() && NetDriver)
 	{
+		ensureMsgf(NetDriver->IsA<UMultiServerNetDriver>(), TEXT("Multi-server beacon NetDriver should be a subclass of UMultiServerNetDriver to function correctly. Check the NetDriverDefinition for MultiServerNetDriver."));
+
 		NetDriver->SetUsingRemoteObjectReferences(!!UE_WITH_REMOTE_OBJECT_HANDLE && GMultiServerAllowRemoteObjectReferences);
 		return true;
 	}
diff --git a/Engine/Plugins/Runtime/MultiServerReplication/Source/MultiServerReplication/Private/MultiServerBeaconHost.cpp b/Engine/Plugins/Runtime/MultiServerReplication/Source/MultiServerReplication/Private/MultiServerBeaconHost.cpp
--- a/Engine/Plugins/Runtime/MultiServerReplication/Source/MultiServerReplication/Private/MultiServerBeaconHost.cpp
+++ b/Engine/Plugins/Runtime/MultiServerReplication/Source/MultiServerReplication/Private/MultiServerBeaconHost.cpp
@@ -1,6 +1,7 @@
 // Copyright Epic Games, Inc. All Rights Reserved.
 
 #include "MultiServerBeaconHost.h"
+#include "MultiServerNetDriver.h"
 #include "Misc/CommandLine.h"
 #include "Net/DataChannel.h"
 #include "Engine/NetConnection.h"
@@ -32,6 +33,8 @@ bool AMultiServerBeaconHost::InitHost()
 	{
 		if (InitBase() && NetDriver)
 		{
+			ensureMsgf(NetDriver->IsA<UMultiServerNetDriver>(), TEXT("Multi-server beacon NetDriver should be a subclass of UMultiServerNetDriver to function correctly. Check the NetDriverDefinition for MultiServerNetDriver."));
+
 			FString Error;
 			if (NetDriver->InitListen(this, URL, bReuseAddressAndPort, Error))
 			{
diff --git a/Engine/Plugins/Runtime/MultiServerReplication/Source/MultiServerReplication/Private/MultiServerNetDriver.cpp b/Engine/Plugins/Runtime/MultiServerReplication/Source/MultiServerReplication/Private/MultiServerNetDriver.cpp
new file mode 100644
--- /dev/null
+++ b/Engine/Plugins/Runtime/MultiServerReplication/Source/MultiServerReplication/Private/MultiServerNetDriver.cpp
@@ -0,0 +1,13 @@
+// Copyright Epic Games, Inc. All Rights Reserved.
+
+#include "MultiServerNetDriver.h"
+
+#include UE_INLINE_GENERATED_CPP_BY_NAME(MultiServerNetDriver)
+
+void UMultiServerNetDriver::SetWorld(UWorld* InWorld)
+{
+	Super::SetWorld(InWorld);
+
+	// Unsubscribe from world tick events, these drivers are ticked manually outside of the normal world tick.
+	UnregisterTickEvents(InWorld);
+}
\ No newline at end of file
diff --git a/Engine/Plugins/Runtime/MultiServerReplication/Source/MultiServerReplication/Private/MultiServerNetDriver.h b/Engine/Plugins/Runtime/MultiServerReplication/Source/MultiServerReplication/Private/MultiServerNetDriver.h
new file mode 100644
--- /dev/null
+++ b/Engine/Plugins/Runtime/MultiServerReplication/Source/MultiServerReplication/Private/MultiServerNetDriver.h
@@ -0,0 +1,32 @@
+// Copyright Epic Games, Inc. All Rights Reserved.
+
+#pragma once
+
+#include "IpNetDriver.h"
+#include "MultiServerNetDriver.generated.h"
+
+// Multi-server nodes need to use UMultiServerNetDriver (or a subclass) because they control the ticking
+// of their NetDrivers directly instead of letting the world tick them.
+// 
+// Use this MultiServerNetDriver as the DriverClassName for the MultiServerNetDriver definition in
+// your project's NetDriverDefinitions engine config, for example:
+// 
+// [/Script/Engine.Engine]
+// +NetDriverDefinitions=(DefName="MultiServerNetDriver", DriverClassName="/Script/MultiServerReplication.MultiServerNetDriver", DriverClassNameFallback="/Script/MultiServerReplication.MultiServerNetDriver")
+// 
+// This allows control over the timing of the Tick(Flush|Dispatch) and PostTick(Flush|Dispatch) functions
+// to ensure they're always called as atomic units. Since MultiServer drivers might be ticked from within
+// a NetDriver that's being ticked by the world, and the world ticks in passes (all netdrivers Tick, then
+// all netdrivers PostTick), we could end up in a situation where a MultiServer driver has Ticked, and is
+// Ticked again before the corresponding PostTick was called (if the world was allowed to tick
+// the MultiServer drivers). This is not compatible with Iris replication.
+// 
+// We use a NetDriver subclass so we can override SetWorld and undo the normal world tick registration.
+UCLASS(Transient)
+class UMultiServerNetDriver : public UIpNetDriver
+{
+	GENERATED_BODY()
+
+public:
+	virtual void SetWorld(UWorld* InWorld) override;
+};
diff --git a/Engine/Plugins/Runtime/MultiServerReplication/Source/MultiServerReplication/Private/MultiServerNode.cpp b/Engine/Plugins/Runtime/MultiServerReplication/Source/MultiServerReplication/Private/MultiServerNode.cpp
--- a/Engine/Plugins/Runtime/MultiServerReplication/Source/MultiServerReplication/Private/MultiServerNode.cpp
+++ b/Engine/Plugins/Runtime/MultiServerReplication/Source/MultiServerReplication/Private/MultiServerNode.cpp
@@ -39,6 +39,7 @@ UMultiServerNode* UMultiServerNode::Create(const FMultiServerNodeCreateParams& P
 
 	if (bResult)
 	{
+		NewNode->RegisterTickEvents();
 		return NewNode;
 	}
 	else
@@ -50,6 +51,13 @@ UMultiServerNode* UMultiServerNode::Create(const FMultiServerNodeCreateParams& P
 	return nullptr;
 }
 
+void UMultiServerNode::BeginDestroy()
+{
+	UnregisterTickEvents();
+
+	Super::BeginDestroy();
+}
+
 bool UMultiServerNode::RegisterServer(const FMultiServerNodeCreateParams& Params)
 {
 	if (Params.World == nullptr)
@@ -151,6 +159,50 @@ bool UMultiServerNode::RegisterServer(const FMultiServerNodeCreateParams& Params
 	return true;
 }
 
+void UMultiServerNode::RegisterTickEvents()
+{
+	UWorld* World = GetWorld();
+	if (World)
+	{
+		TickDispatchDelegateHandle = World->OnTickDispatch().AddUObject(this, &UMultiServerNode::InternalTickDispatch);
+		TickFlushDelegateHandle = World->OnTickFlush().AddUObject(this, &UMultiServerNode::InternalTickFlush);
+	}
+}
+
+void UMultiServerNode::UnregisterTickEvents()
+{
+	UWorld* World = GetWorld();
+	if (World)
+	{
+		World->OnTickDispatch().Remove(TickDispatchDelegateHandle);
+		World->OnTickFlush().Remove(TickFlushDelegateHandle);
+	}
+}
+
+void UMultiServerNode::InternalTickDispatch(float DeltaSeconds)
+{
+	ForEachNetDriver([DeltaSeconds](UNetDriver* NetDriver)
+		{
+			if (NetDriver->GetWorld())
+			{
+				NetDriver->TickDispatch(DeltaSeconds);
+				NetDriver->PostTickDispatch();
+			}
+		});
+}
+
+void UMultiServerNode::InternalTickFlush(float DeltaSeconds)
+{
+	ForEachNetDriver([DeltaSeconds](UNetDriver* NetDriver)
+		{
+			if (NetDriver->GetWorld())
+			{
+				NetDriver->TickFlush(DeltaSeconds);
+				NetDriver->PostTickFlush();
+			}
+		});
+}
+
 AMultiServerBeaconClient* UMultiServerNode::GetBeaconClientForRemotePeer(FStringView RemotePeerId)
 {
 	// See if we are the host of the target server
diff --git a/Engine/Plugins/Runtime/MultiServerReplication/Source/MultiServerReplication/Public/MultiServerNode.h b/Engine/Plugins/Runtime/MultiServerReplication/Source/MultiServerReplication/Public/MultiServerNode.h
--- a/Engine/Plugins/Runtime/MultiServerReplication/Source/MultiServerReplication/Public/MultiServerNode.h
+++ b/Engine/Plugins/Runtime/MultiServerReplication/Source/MultiServerReplication/Public/MultiServerNode.h
@@ -66,6 +66,8 @@ public:
 
 	static void ParseCommandLineIntoCreateParams(FMultiServerNodeCreateParams& InOutParams);
 
+	virtual void BeginDestroy();
+
 	bool RegisterServer(const FMultiServerNodeCreateParams& Params);
 
 	AMultiServerBeaconClient* GetBeaconClientForRemotePeer(FStringView RemotePeerId);
@@ -114,6 +116,23 @@ private:
 	float RetryConnectMaxDelay;
 
 	TSubclassOf<AMultiServerBeaconClient> UserBeaconClass;
+
+	// Multi-server nodes handle ticking of their NetDrivers directly instead of letting the world tick them.
+	// This allows us to control the timing of the Tick(Flush|Dispatch) and PostTick(Flush|Dispatch) functions
+	// to ensure they're always called as atomic units. Since MultiServer drivers might be ticked from within
+	// a NetDriver that's being ticked by the world, and the world ticks in passes (all netdrivers Tick, then
+	// all netdrivers PostTick), we could end up in a situation where a MultiServer driver has Ticked, and is
+	// Ticked again before the corresponding PostTick was called (if the world was allowed to tick
+	// the MultiServer drivers).
+	// Note UMultiServerNetDriver::SetWorld unregisters from world tick events. See UMultiServerNetDriver.
+	FDelegateHandle TickDispatchDelegateHandle;
+	FDelegateHandle TickFlushDelegateHandle;
+
+	void RegisterTickEvents();
+	void UnregisterTickEvents();
+
+	void InternalTickDispatch(float DeltaSeconds);
+	void InternalTickFlush(float DeltaSeconds);
 };
 
 template<class T>
diff --git a/Engine/Source/Runtime/Engine/Private/Actor.cpp b/Engine/Source/Runtime/Engine/Private/Actor.cpp
--- a/Engine/Source/Runtime/Engine/Private/Actor.cpp
+++ b/Engine/Source/Runtime/Engine/Private/Actor.cpp
@@ -1290,9 +1290,13 @@ void AActor::PostMigrate(const FUObjectMigrationContext& MigrationContext)
 			// Migrated actors are unique to the server instance that currently owns their memory and are not replicated to other servers.
 			if (World)
 			{
+				// Notify replication systems that this actor has been destroyed.
+#if UE_WITH_IRIS
+				EndReplication(EEndPlayReason::Destroyed);
+#endif // UE_WITH_IRIS
+
 				UE::Net::Private::FNetPropertyConditionManager::Get().NotifyObjectDestroyed(this);
 
-				// Notify net drivers that this actor has been destroyed.
 				if (FWorldContext* Context = GEngine->GetWorldContextFromWorld(World))
 				{
 					for (FNamedNetDriver& Driver : Context->ActiveNetDrivers)
