MultiServerProxy:  Send the GameLoginOptions when we do a Join to match the existing PendingNetGame/SplitJoin implementation.

#rnx

diff --git a/Engine/Plugins/Runtime/MultiServerReplication/Source/MultiServerReplication/Private/MultiServerProxy.cpp b/Engine/Plugins/Runtime/MultiServerReplication/Source/MultiServerReplication/Private/MultiServerProxy.cpp
--- a/Engine/Plugins/Runtime/MultiServerReplication/Source/MultiServerReplication/Private/MultiServerProxy.cpp
+++ b/Engine/Plugins/Runtime/MultiServerReplication/Source/MultiServerReplication/Private/MultiServerProxy.cpp
@@ -147,7 +147,7 @@ void UProxyListenerNotify::ConnectToGameServer(UNetConnection* ProxyConnection,
 
 	// Add a player to use on the game server.
 	const FPlatformUserId GameServerClientId = FPlatformUserId::CreateFromInternalId(ProxyNetDriver->GetNextGameServerClientId());
-	ULocalPlayer* NewPlayer = NewObject<ULocalPlayer>(GEngine, ULocalPlayer::StaticClass());
+	ULocalPlayer* NewPlayer = NewObject<ULocalPlayer>(GEngine, GEngine->LocalPlayerClass);
 	ProxyNetDriver->GetWorld()->GetGameInstance()->AddLocalPlayer(NewPlayer, GameServerClientId);
 
 	// The new player will use the same unique identifier as the incoming proxy connection so that it will be propogated up to 
@@ -165,6 +165,12 @@ void UProxyListenerNotify::ConnectToGameServer(UNetConnection* ProxyConnection,
 		FString URLStr = FString::Format(TEXT("127.0.0.1:{0}"), { GameServerPort });
 		FURL URL(nullptr, *URLStr, TRAVEL_Absolute);
 		
+		const FString GameLoginOptions = NewPlayer->GetGameLoginOptions();
+		if (!GameLoginOptions.IsEmpty())
+		{
+			URL.AddOption(*GameLoginOptions);
+		}
+
 		URL.AddOption(*FString::Printf(TEXT("HandshakeId=%u"), ClientHandshakeId));
 
 		// Maybe this should be a control message since it changes the server setting (it's global driver setting).
