The multiserver proxy shares a single NetGUID cache; game server use FRemoteObjectId as a NetGUID; and objects are cached on the proxy when migrating between servers.

diff --git a/Engine/Plugins/Runtime/MultiServerReplication/Source/MultiServerReplication/Private/MultiServerProxy.cpp b/Engine/Plugins/Runtime/MultiServerReplication/Source/MultiServerReplication/Private/MultiServerProxy.cpp
--- a/Engine/Plugins/Runtime/MultiServerReplication/Source/MultiServerReplication/Private/MultiServerProxy.cpp
+++ b/Engine/Plugins/Runtime/MultiServerReplication/Source/MultiServerReplication/Private/MultiServerProxy.cpp
@@ -26,6 +26,60 @@ FString LexToString(EProxyConnectionState State)
 	}
 }
 
+FProxyNetGUIDCache::FProxyNetGUIDCache(UNetDriver* NetDriver) : FNetGUIDCache(NetDriver)
+{
+
+}
+
+FNetworkGUID FProxyNetGUIDCache::AssignNewNetGUID_Server(UObject* Object)
+{
+	UProxyNetDriver* ProxyNetDriver = Cast<UProxyNetDriver>(Driver);
+	
+	if (ensureMsgf(ProxyNetDriver, TEXT("FProxyNetGUIDCache is only intended to be used by UProxyNetDriver")))
+	{
+		TSharedPtr<FNetGUIDCache> BackendGuidCache = ProxyNetDriver->GetSharedBackendNetGuidCache();
+		if (ensure(BackendGuidCache))
+		{
+			FNetworkGUID MatchingNetGUID = BackendGuidCache->GetNetGUID(Object);
+			if (MatchingNetGUID.IsValid())
+			{
+				// Catch cases where the NetGUID is already in the proxy cache.
+				ensure(!GetNetGUID(Object).IsValid());
+
+				UE_LOG(LogNetProxy, Log, TEXT("NetGUID: Registering NetGUID %s for object %s found in the shared NetGUID cache."), *MatchingNetGUID.ToString(), *Object->GetName());
+				RegisterNetGUID_Server(MatchingNetGUID, Object);
+				return MatchingNetGUID;
+			}
+			else
+			{
+				UE_LOG(LogNetProxy, Error, TEXT("NetGUID: Failed to find a NetGUID for object %s in the shared NetGUID cache."), *Object->GetName());
+				return FNetworkGUID::GetDefault();
+			}
+		}
+	}
+
+	return FNetworkGUID::GetDefault();
+}
+
+FNetworkGUID FProxyNetGUIDCache::AssignNewNetGUIDFromPath_Server(const FString& PathName, UObject* ObjOuter, UClass* ObjClass)
+{
+	ensureMsgf(0, TEXT("UProxyNetDriver does not support assinging NetGUIDs from a path."));
+	return FNetworkGUID::GetDefault();
+}
+
+FProxyBackendNetGUIDCache::FProxyBackendNetGUIDCache(UNetDriver* NetDriver) : FNetGUIDCache(NetDriver)
+{
+
+}
+
+bool FProxyBackendNetGUIDCache::IsNetGUIDAuthority() const
+{
+	// The backend NetGUID cache will always not be a GUID authority (i.e. it won't assign new GUIDs). This
+	// is required because the default implementation queries the associated NetDriver which will be UProxyNetDriver
+	// and a listening server (i.e. this function would return true).
+	return false;
+}
+
 void UGameServerNotify::NotifyAcceptedConnection(class UNetConnection* Connection)
 {
 
@@ -137,10 +191,13 @@ void UProxyListenerNotify::ConnectToGameServer(UNetConnection* ProxyConnection,
 		GameServerConnectionState->NetDriverName = FName(*FString::Printf(TEXT("ProxyToGameServer-%d"), GameServerDriverId++));
 
 		GEngine->CreateNamedNetDriver(ProxyNetDriver->GetWorld(), GameServerConnectionState->NetDriverName, "ProxyBackendNetDriver");
-		GameServerConnectionState->NetDriver = GEngine->FindNamedNetDriver(GameServerConnectionState->World, GameServerConnectionState->NetDriverName);
+		GameServerConnectionState->NetDriver = Cast<UProxyBackendNetDriver>(GEngine->FindNamedNetDriver(GameServerConnectionState->World, GameServerConnectionState->NetDriverName));
+
+		ensure(GameServerConnectionState->NetDriver);
 		
 		GameServerConnectionState->NetDriver->SetWorld(ProxyNetDriver->GetWorld());
-		Cast<UProxyBackendNetDriver>(GameServerConnectionState->NetDriver)->SetProxyNetDriver(ProxyNetDriver);
+		GameServerConnectionState->NetDriver->SetProxyNetDriver(ProxyNetDriver);
+		GameServerConnectionState->NetDriver->GuidCache = ProxyNetDriver->GetSharedBackendNetGuidCache();
 
 		UE_LOG(LogNetProxy, Log, TEXT("Created a Proxy Game Server NetDriver (name=%s, url=%s)"), *GameServerConnectionState->NetDriver->GetName(), *GameServerConnectionState->GameServerURL.ToString());
 	}
@@ -150,6 +207,8 @@ void UProxyListenerNotify::ConnectToGameServer(UNetConnection* ProxyConnection,
 	ULocalPlayer* NewPlayer = NewObject<ULocalPlayer>(GEngine, GEngine->LocalPlayerClass);
 	ProxyNetDriver->GetWorld()->GetGameInstance()->AddLocalPlayer(NewPlayer, GameServerClientId);
 
+	ProxyNetDriver->ProxyPlayerToProxyConnection.Add(NewPlayer, ProxyConnection);
+
 	// The new player will use the same unique identifier as the incoming proxy connection so that it will be propogated up to 
 	// the game servers through UNetConnection::PlayerId. This way each game server's incoming connection will have a PlayerId 
 	// that corresponds to a client connected to the proxy.
@@ -376,47 +435,6 @@ void UProxyBackendNetDriver::InternalProcessRemoteFunction(
 {
 	bool bShouldForwardRPC = true;
 
-	// The majority of the actors on the proxy are replicated from the game servers so their owners will be game server player
-	// controllers. The exception are proxy player controllers which are spawned by the proxy and will have an owner that is
-	// the proxy player controller. In this case we must forward the RPC to the primary game server player controller.
-	if (ProxyNetDriver->IsProxySpawned(Actor))
-	{
-		AActor* MappedActor = Actor;
-
-		if (APlayerController* ProxyPlayerController = Cast<APlayerController>(MappedActor))
-		{
-			if (ULocalPlayer* Player = Cast<ULocalPlayer>(ProxyPlayerController->Player))
-			{
-				if (TObjectPtr<APlayerController>* GameServerPlayerControllerPtr = ProxyNetDriver->ProxyPlayerToPrimaryGameServerPlayerController.Find(Player))
-				{
-					UNetConnection* GameServerConnection = GameServerPlayerControllerPtr->Get()->NetConnection;
-					if (ensure(GameServerConnection->PlayerController))
-					{
-						ensure(GameServerConnection->PlayerController->Player == Player);
-						MappedActor = GameServerConnection->PlayerController;
-
-						UE_LOG(LogNetProxy, VeryVerbose, TEXT("RPC: Remapping proxy player controller %s to game server player controller %s when forwarding RPC %s on %s"),
-							   *ProxyPlayerController->GetName(),
-							   *GameServerConnection->PlayerController->GetName(),
-							   *Function->GetName(),
-							   *Actor->GetName());
-					}
-				}
-			}
-		}
-		
-		if (ProxyNetDriver->IsProxySpawned(MappedActor))
-		{
-			bShouldForwardRPC = false;
-
-			UE_LOG(LogNetProxy, Warning, TEXT("RPC: Unable to map proxy owned actor %s to game server %s owned actor."), *Actor->GetName(), *Connection->GetName());
-		}
-		else
-		{
-			Actor = MappedActor;
-		}
-	}
-
 	// If the sub-object is not owned by the actor, attempt to find an component in that actor that matches the same type.
 	// This logic assumes that an actor only has one component of a given type and will fail if that assumption is incorrect.
 	if (UActorComponent* SubObjectAsActorComponent = Cast<UActorComponent>(SubObject))
@@ -485,6 +503,15 @@ void UProxyBackendNetDriver::InternalProcessRemoteFunction(
 	}
 }
 
+bool UProxyBackendNetDriver::ShouldClientDestroyActor(AActor* Actor, EChannelCloseReason CloseReason) const
+{
+	// If an actor is destroyed remotely on a game server because it's being migrated to another server then it shouldn't be
+	// destroyed on the proxy when removed from the replication system of the originating server. When the actor arrives on
+	// the destination server it will be added to that server's replication system and the actor will be re-used on the proxy
+	// since it's still referenced in the shared backend NetGUID cache.
+	return (CloseReason != EChannelCloseReason::Migrated);
+}
+
 void UProxyNetDriver::RegisterGameServer(const FURL& GameServerURL)
 {
 	UE_LOG(LogNetProxy, Log, TEXT("Registering Game Server for URL: %s"), *GameServerURL.ToString());
@@ -563,6 +590,12 @@ bool UProxyNetDriver::InitBase(bool bInitAsClient, FNetworkNotify* InNotify, con
 	
 	DisableActorLogicAndGameCode();
 
+	GuidCache = TSharedPtr<FNetGUIDCache>(new FProxyNetGUIDCache(this));
+
+	// All of the instances of UProxyBackendNetDriver will use this shared GUID cache to ensure that objects
+	// with the same NetGUID use the same object pointers and don't create duplicate objects for the same NetGUID.
+	SharedBackendNetGuidCache = TSharedPtr<FNetGUIDCache>(new FProxyBackendNetGUIDCache(this));
+
 	return bSuccess;
 }
 
@@ -601,12 +634,20 @@ void UProxyNetDriver::ForwardRemoteFunction(UObject* RootObject, UObject* SubObj
 		UNetConnection* GameConnection = GamePlayerControllerPtr->Get()->NetConnection;
 		if (GameConnection)
 		{
-			GameConnection->Driver->ProcessRemoteFunction(OwningActor, Function, Parms, static_cast<FOutParmRec*>(nullptr), static_cast<FFrame*>(nullptr), SubObject);
+			const bool bIsConnectionToGameServer = GameConnection->IsA(UProxyBackendNetConnection::StaticClass()) || GameConnection->IsA(UProxyBackendChildNetConnection::StaticClass());
+			if (ensure(bIsConnectionToGameServer))
+			{
+				GameConnection->Driver->ProcessRemoteFunction(OwningActor, Function, Parms, static_cast<FOutParmRec*>(nullptr), static_cast<FFrame*>(nullptr), SubObject);
+			}
+			else
+			{
+				UE_LOG(LogNetProxy, Error, TEXT("RPC: Unable to forward %s on %s received from a proxy connection because it's going to be routed back to the proxy connection"), *Function->GetName(), *RootObject->GetName());
+			}
 		}
 	}
 	else
 	{
-		UE_LOG(LogNetProxy, Warning, TEXT("RPC: Unable to forward %s on %s because player %s isn't mapped to a primary game server."), *Function->GetName(), *OwningActor->GetName(), *OwningPlayer->GetName());
+		UE_LOG(LogNetProxy, Warning, TEXT("RPC: Unable to forward %s on %s because player %s isn't mapped to a primary game server."), *Function->GetName(), *RootObject->GetName(), *OwningPlayer->GetName());
 	}
 }
 
@@ -638,9 +679,7 @@ void UProxyNetDriver::AddNetworkActor(AActor* Actor)
 	// loaded by the proxy and those replicated from the connected game servers. It's assumed that 
 	// the actors replicated from the game servers will not have a role of ROLE_Authority and will
 	// therefore be unaffacted by this code and replicate as normal.
-	//
-	// Actors that are explicitly spawned as part of the proxy functionality are allowed to be replicated.
-	if (Actor->GetIsReplicated() && !IsProxySpawned(Actor))
+	if (Actor->GetIsReplicated())
 	{
 		if (Actor->GetLocalRole() == ROLE_Authority)
 		{
@@ -652,6 +691,19 @@ void UProxyNetDriver::AddNetworkActor(AActor* Actor)
 	Super::AddNetworkActor(Actor);
 }
 
+void UProxyNetDriver::RemoveNetworkActor(AActor* Actor)
+{
+	Super::RemoveNetworkActor(Actor);
+
+	// Remove this actor from the proxy netguid cache. On a normal dedicated server there's a one-to-one
+	// mapping between a NetGUID and the corresponding object - i.e. each unique object has a unique 
+	// netguid. On the proxy an object can be created/destroyed repeatedly as it's migrated between
+	// servers but will still maintain the same netguid (derived from a remote id). For this reason,
+	// it's important to remove an object from the proxy netguid cache when it's being removed from
+	// the replication system to avoid any conflicts in the cache which trigger checks/ensures.
+	GuidCache->RemoveNetGUID(Actor);
+}
+
 bool UProxyNetDriver::ShouldCallRemoteFunction(UObject* Object, UFunction* Function, const FReplicationFlags& RepFlags) const
 {
 	return !RepFlags.bIgnoreRPCs;
@@ -673,7 +725,8 @@ void UProxyNetDriver::InternalProcessRemoteFunction(
 		// be routed to the game client; other RPCs will be ignored.
 		if (UNetConnection* ProxyConnection = GetProxyConnectionFromPrimaryPlayer(Player))
 		{
-			UE_LOG(LogNetProxy, Verbose, TEXT("RPC: Forwarding %s on %s to proxy connection %s for player %s"), *Function->GetName(), *Actor->GetName(), *ProxyConnection->GetName(), *Player->GetName());
+			UPlayer* Owner = Actor->GetNetOwningPlayerAnyRole();
+			UE_LOG(LogNetProxy, Verbose, TEXT("RPC: Forwarding %s on %s (owner:%s) to proxy connection %s:%s for player %s"), *Function->GetName(), *Actor->GetName(), *GetNameSafe(Owner), *ProxyConnection->Driver->GetName(), *ProxyConnection->GetName(), *Player->GetName());
 			Super::InternalProcessRemoteFunction(Actor, SubObject, ProxyConnection, Function, Parms, OutParms, Stack, bIsServer);
 		}
 		else
@@ -731,7 +784,7 @@ bool UProxyNetDriver::CanDowngradeActorRole(UNetConnection* ProxyConnection, AAc
 		// as the attached proxy player controller, don't downgrade from ROLE_AutonomousProxy to ROLE_SimulatedProxy.
 		if (Actor->GetRemoteRole() == ROLE_AutonomousProxy)
 		{
-			if (UNetConnection* PlayerProxyConnection = GetProxyConnectionFromPrimaryPlayer(Player))
+			if (const UNetConnection* PlayerProxyConnection = const_cast<UProxyNetDriver*>(this)->GetProxyConnectionFromPrimaryPlayer(Player))
 			{
 				if (PlayerProxyConnection == ProxyConnection)
 				{
@@ -782,11 +835,6 @@ void UProxyNetDriver::DisableActorLogicAndGameCode()
 #endif
 }
 
-bool UProxyNetDriver::IsProxySpawned(AActor* Actor) const
-{
-	return ProxySpawnedActors.Contains(Actor);
-}
-
 void UProxyNetDriver::GameServerAssignPlayerController(UNetConnection* ChildGameServerConnection, UNetConnection* NotUsedParentGameServerConnection, APlayerController* GameServerPlayerController)
 {
 	UE_LOG(LogNetProxy, Log, TEXT("Received a new player controller %s:%s for connection %s:%s (client_handshake_id=%u)."), 
@@ -795,7 +843,6 @@ void UProxyNetDriver::GameServerAssignPlayerController(UNetConnection* ChildGame
 		   GameServerPlayerController->GetClientHandshakeId());
 	
 	// The game server player controller is used as a client connection to the game server so must be ROLE_AutonomousProxy.
-	// NOTE: A proxy player controller will still be ROLE_Authority since it's spawned by the proxy.
 	GameServerPlayerController->SetRole(ROLE_AutonomousProxy);
 
 	const uint32 ClientHandshakeId = GameServerPlayerController->GetClientHandshakeId();
@@ -823,6 +870,17 @@ void UProxyNetDriver::GameServerAssignPlayerController(UNetConnection* ChildGame
 	// is the primary game server for the proxy connection (decided later in this function).
 	Player->PlayerController = nullptr;
 	
+	// Since an existing player controller can be re-assigned it must be detached from any existing connections and players.
+	if (GameServerPlayerController->NetConnection != nullptr)
+	{
+		UE_LOG(LogNetProxy, Log, TEXT("Player controller %s already assigned to connection %s:%s."),
+			   *GameServerPlayerController->GetName(),
+			   *GetNameSafe(GameServerPlayerController->NetConnection->Driver),
+			   *GameServerPlayerController->NetConnection->GetName());
+
+		DetachPlayerController(GameServerPlayerController);
+	}
+
 	// If a new player controller is spawned on the proxy from a game server it's assumed it will replace any
 	// existing game server player controllers used by connections and players on the proxy. For this reason
 	// always detach an existing game server player controllers from the connection.
@@ -862,11 +920,14 @@ void UProxyNetDriver::GameServerAssignPlayerController(UNetConnection* ChildGame
 							{
 								UPlayer* PreviousPlayer = GameServerPlayerController->Player;
 			
-								APlayerController* ProxyPlayerController = SpawnPlayerController();
-								if (ensure(ProxyPlayerController))
+								APlayerController* ProxyPlayerController = GameServerPlayerController;
 								{
-									DetachAndDestroyPlayerController(ProxyConnection->PlayerController);
-									AttachPlayerController(ProxyPlayerController, ProxyConnection, Player);
+									DetachPlayerController(ProxyConnection->PlayerController);
+
+									ProxyConnection->SetConnectionState(EConnectionState::USOCK_Open);
+									ProxyConnection->SetClientHandshakeId(ProxyPlayerController->GetClientHandshakeId());
+									ProxyConnection->PlayerController = ProxyPlayerController;
+									ProxyConnection->OwningActor = ProxyPlayerController;
 
 									// The player always points to the proxy player controller if that player is
 									// considered the primary player for the proxy connection.
@@ -905,8 +966,11 @@ void UProxyNetDriver::GameServerAssignPlayerController(UNetConnection* ChildGame
 				{
 					if (!GameServerPlayerController->IsA(ANoPawnPlayerController::StaticClass()))
 					{
-						if (ULocalPlayer* ExistingPrimaryPlayer = Cast<ULocalPlayer>(ProxyConnection->PlayerController->Player))
+						TObjectPtr<ULocalPlayer>* ExistingPrimaryPlayerPtr = ProxyConnectionToPrimaryPlayer.Find(ProxyConnection);
+						if (ExistingPrimaryPlayerPtr)
 						{
+							ULocalPlayer* ExistingPrimaryPlayer = *ExistingPrimaryPlayerPtr;
+
 							if (ProxyPlayerToPrimaryGameServerPlayerController.Contains(ExistingPrimaryPlayer))
 							{
 								int32* PreviousPrimaryGameServerIdxPtr = ProxyPlayerToPrimaryGameServer.Find(ExistingPrimaryPlayer);
@@ -928,7 +992,7 @@ void UProxyNetDriver::GameServerAssignPlayerController(UNetConnection* ChildGame
 					
 										ProxyPlayerToPrimaryGameServerPlayerController.Remove(ExistingPrimaryPlayer);
 
-										DeregisterPrimaryGameServer(ProxyConnection);
+										DeregisterPrimaryGameServer(ProxyConnection, ExistingPrimaryPlayer);
 										RegisterPrimaryGameServer(ProxyConnection, Player, NewPrimaryGameServerIdx);
 				
 										ExistingPrimaryPlayer->PlayerController = nullptr;
@@ -948,8 +1012,8 @@ void UProxyNetDriver::GameServerAssignPlayerController(UNetConnection* ChildGame
 	}
 			
 	// Update the internal connection mapping to use the new child game server connection (if it's different from the parent connection).
-	ClientHandshakeIdToInternalRoute.Remove(NextClientHandshakeId);
-	ClientHandshakeIdToInternalRoute.Add(NextClientHandshakeId, Route);
+	ClientHandshakeIdToInternalRoute.Remove(ClientHandshakeId);
+	ClientHandshakeIdToInternalRoute.Add(ClientHandshakeId, Route);
 
 	// If this was a parent connection send any join requests for any pending multiplexed connections.
 	if (ChildGameServerConnection == Route.ParentGameServerConnection)
@@ -971,6 +1035,11 @@ void UProxyNetDriver::GameServerAssignPlayerController(UNetConnection* ChildGame
 	}
 }
 
+TSharedPtr<FNetGUIDCache> UProxyNetDriver::GetSharedBackendNetGuidCache()
+{
+	return SharedBackendNetGuidCache;
+}
+
 void UProxyNetDriver::FlushSplitJoinRequests(FGameServerConnectionState* GameServerConnectionState)
 {
 	EConnectionState ParentConnectionState = GameServerConnectionState->NetDriver->ServerConnection->GetConnectionState();
@@ -1014,24 +1083,21 @@ int32 UProxyNetDriver::GetGameServerWithPlayer(ULocalPlayer* Player)
 	return -1;
 }
 
-void UProxyNetDriver::DeregisterPrimaryGameServer(UNetConnection* ProxyConnection)
+void UProxyNetDriver::DeregisterPrimaryGameServer(UNetConnection* ProxyConnection, ULocalPlayer* PreviousPrimaryPlayer)
 {
 	if (ensure(ProxyConnection->PlayerController))
 	{
-		ULocalPlayer* OldPlayer = Cast<ULocalPlayer>(ProxyConnection->PlayerController->Player);
-		if (OldPlayer)
-		{
-			FUniqueNetIdRepl& PlayerId = ProxyConnection->PlayerId;
+		FUniqueNetIdRepl& PlayerId = ProxyConnection->PlayerId;
 
-			UE_LOG(LogNetProxy, Log, TEXT("Clearing player %s primary game server (proxy_connection=%s:%s)"),
-				   *OldPlayer->GetName(),
-				   *ProxyConnection->GetDriver()->GetName(),
-				   *ProxyConnection->GetName());
+		UE_LOG(LogNetProxy, Log, TEXT("Clearing player %s primary game server (proxy_connection=%s:%s)"),
+			*PreviousPrimaryPlayer->GetName(),
+			*GetNameSafe(ProxyConnection->GetDriver()),
+			*ProxyConnection->GetName());
 
-			ensure(ProxyPlayerToPrimaryGameServer.Contains(OldPlayer));
+		ensure(ProxyPlayerToPrimaryGameServer.Contains(PreviousPrimaryPlayer));
 		
-			ProxyPlayerToPrimaryGameServer.Remove(OldPlayer);
-		}
+		ProxyPlayerToPrimaryGameServer.Remove(PreviousPrimaryPlayer);
+		ProxyConnectionToPrimaryPlayer.Remove(ProxyConnection);
 	}
 }
 
@@ -1050,6 +1116,7 @@ void UProxyNetDriver::RegisterPrimaryGameServer(UNetConnection* ProxyConnection,
 		ensure(!ProxyPlayerToPrimaryGameServer.Contains(PrimaryPlayer));
 
 		ProxyPlayerToPrimaryGameServer.Add(PrimaryPlayer, GameServerConnectionStateIndex);
+		ProxyConnectionToPrimaryPlayer.Add(ProxyConnection, PrimaryPlayer);
 	}
 }
 
@@ -1091,13 +1158,8 @@ void UProxyNetDriver::DetachPlayerController(APlayerController* PlayerController
 	{
 		UPlayer* Player = PlayerController->Player;
 
-		UE_LOG(LogNetProxy, Log, TEXT("Detaching old controller %s:%s on %s attached to player %s."),
-			*PlayerController->GetName(),
-			*PlayerController->GetClass()->GetName(),
-			*GetName(),
-			*Player->GetName());
-
-		if (UNetConnection* Connection = PlayerController->NetConnection)
+		UNetConnection* Connection = PlayerController->NetConnection;
+		if (PlayerController->NetConnection)
 		{
 			Connection->PlayerController = nullptr;
 			Connection->OwningActor = nullptr;
@@ -1107,66 +1169,20 @@ void UProxyNetDriver::DetachPlayerController(APlayerController* PlayerController
 		PlayerController->Player = nullptr;
 
 		// Only detach the player controller from the player if it's the one referenced by the player.
-		if (Player->PlayerController == PlayerController)
-		{
-			Player->PlayerController = nullptr;
-		}
-	}
-}
-
-void UProxyNetDriver::DetachAndDestroyPlayerController(APlayerController* PlayerController)
-{
-	DetachPlayerController(PlayerController);
-	DestroyPlayerController(PlayerController);
-}
-
-APlayerController* UProxyNetDriver::SpawnPlayerController()
-{
-	AGameModeBase* GameMode = GetWorld()->GetAuthGameMode();
-
-	if (ensure(GameMode))
-	{
-		FActorSpawnParameters SpawnInfo;
-		SpawnInfo.ObjectFlags |= RF_Transient;
-		SpawnInfo.bDeferConstruction = true;
-	
-		APlayerController* PlayerController = GetWorld()->SpawnActor<APlayerController>(GameMode->PlayerControllerClass, SpawnInfo);
-
-		if (PlayerController)
-		{
-			UE_LOG(LogNetProxy, Log, TEXT("Spawning player controller %s of type %s."), *PlayerController->GetName(), *GameMode->PlayerControllerClass->GetName());
-
-			ProxySpawnedActors.Add(PlayerController);
-	
-			GetWorld()->AddController(PlayerController);
-
-			PlayerController->SetRole(ROLE_Authority);
-			PlayerController->SetReplicates(true);
-			PlayerController->SetAutonomousProxy(true); // Sets AActor::RemoteRole.
-			PlayerController->FinishSpawning(FTransform());
-		}
-
-		return PlayerController;
-	}
-
-	return nullptr;
-}
-
-void UProxyNetDriver::DestroyPlayerController(APlayerController* PlayerController)
-{
-	if (PlayerController)
-	{
-		// Only player controllers spawned by the proxy should be destroyed.
-		if (ensure(ProxySpawnedActors.Contains(PlayerController)))
+		if (Player)
 		{
-			if (ensure(PlayerController->GetLocalRole() == ROLE_Authority))
+			if (Player->PlayerController == PlayerController)
 			{
-				ProxySpawnedActors.Remove(PlayerController);
-
-				PlayerController->GetWorld()->RemoveController(PlayerController);
-				PlayerController->GetWorld()->DestroyActor(PlayerController);
+				Player->PlayerController = nullptr;
 			}
 		}
+
+		UE_LOG(LogNetProxy, Log, TEXT("Detaching old controller %s:%s on %s attached to player %s and connection %s."),
+			   *PlayerController->GetName(),
+			   *PlayerController->GetClass()->GetName(),
+			   *GetName(),
+			   *GetNameSafe(Player),
+			   *GetNameSafe(Connection));
 	}
 }
 
@@ -1174,9 +1190,9 @@ UNetConnection* UProxyNetDriver::GetProxyConnectionFromPrimaryPlayer(UPlayer* Pl
 {
 	if (ULocalPlayer* LocalPlayer = Cast<ULocalPlayer>(Player))
 	{
-		if (APlayerController* ProxyPlayerController = LocalPlayer->PlayerController)
+		if (TObjectPtr<UNetConnection>* ConnectionPtr = ProxyPlayerToProxyConnection.Find(LocalPlayer))
 		{
-			return ProxyPlayerController->NetConnection;
+			return ConnectionPtr->Get();
 		}
 	}
 
diff --git a/Engine/Plugins/Runtime/MultiServerReplication/Source/MultiServerReplication/Public/MultiServerProxy.h b/Engine/Plugins/Runtime/MultiServerReplication/Source/MultiServerReplication/Public/MultiServerProxy.h
--- a/Engine/Plugins/Runtime/MultiServerReplication/Source/MultiServerReplication/Public/MultiServerProxy.h
+++ b/Engine/Plugins/Runtime/MultiServerReplication/Source/MultiServerReplication/Public/MultiServerProxy.h
@@ -4,6 +4,7 @@
 
 #include "Engine/PendingNetGame.h"
 #include "Engine/ChildConnection.h"
+#include "Engine/PackageMapClient.h"
 #include "GameFramework/OnlineReplStructs.h"
 #include "IpNetDriver.h"
 #include "IpConnection.h"
@@ -18,6 +19,27 @@ enum class EProxyConnectionState : uint8
 
 FString LexToString(EProxyConnectionState State);
 
+/** 
+ * A NetGUID cache used by UProxyNetDriver that doesn't assign new NetGUIDs but uses an already assigned NetGUID
+ * from a backend server connection's NetGUID cache.
+ */
+class FProxyNetGUIDCache : public FNetGUIDCache
+{
+public:
+	FProxyNetGUIDCache(UNetDriver* NetDriver);
+
+	virtual FNetworkGUID AssignNewNetGUID_Server(UObject* Object);
+	virtual FNetworkGUID AssignNewNetGUIDFromPath_Server(const FString& PathName, UObject* ObjOuter, UClass* ObjClass);
+};
+
+class FProxyBackendNetGUIDCache : public FNetGUIDCache
+{
+public:
+	FProxyBackendNetGUIDCache(UNetDriver* NetDriver);
+
+	virtual bool IsNetGUIDAuthority() const;
+};
+
 USTRUCT()
 struct FGameServerSplitJoinRequest
 {
@@ -32,7 +54,9 @@ struct FGameServerSplitJoinRequest
 };
 
 /**
- * All of the state associated with a proxy server connection to a backend game server.
+ * A physical connection from the proxy server to a single game server.
+ *
+ * All proxy client connections are multiplexed to backend servers through this physical connection.
  */ 
 USTRUCT()
 struct FGameServerConnectionState
@@ -52,7 +76,7 @@ struct FGameServerConnectionState
 	TObjectPtr<class UWorld> World;
 
 	UPROPERTY()
-	TObjectPtr<class UNetDriver> NetDriver;
+	TObjectPtr<class UProxyBackendNetDriver> NetDriver;
 	
 	FName NetDriverName;
 
@@ -71,17 +95,25 @@ private:
 	FGameServerConnectionState() = default;
 };
 
+/** 
+ * A route from a proxy client connection to a single game server.
+ *
+ * This route binds a physical proxy client connection to a logical (multiplex) game server connection.
+ */ 
 USTRUCT()
 struct FMultiServerProxyInternalConnectionRoute
 {
 	GENERATED_BODY()
 
+	/** The client connection to the proxy (the beginning of the route). */
 	UPROPERTY()
 	TObjectPtr<UNetConnection> ProxyConnection;
 
+	/** The parent connection to the backend server. */
 	UPROPERTY()
 	TObjectPtr<UNetConnection> ParentGameServerConnection;
 
+	/** The player used to represent this route in the game systems. */
 	UPROPERTY()
 	TObjectPtr<class ULocalPlayer> Player;
 };
@@ -207,6 +239,7 @@ public:
 	virtual bool ShouldSkipRepNotifies() const override;
 	virtual UChildConnection* CreateChild(UNetConnection* Parent) override;
 	virtual void InternalProcessRemoteFunction(class AActor* Actor, class UObject* SubObject, class UNetConnection* Connection, class UFunction* Function, void* Parms, FOutParmRec* OutParms, FFrame* Stack, bool bIsServer) override;
+	virtual bool ShouldClientDestroyActor(AActor* Actor, EChannelCloseReason CloseReason) const;
 
 private:
 
@@ -289,6 +322,7 @@ public:
 	virtual bool ShouldReplicateFunction(AActor* Actor, UFunction* Function) const override;
 	virtual void NotifyActorChannelOpen(UActorChannel* Channel, AActor* Actor) override;
 	virtual void AddNetworkActor(AActor* Actor) override;
+	virtual void RemoveNetworkActor(AActor* Actor) override;
 	virtual bool ShouldCallRemoteFunction(UObject* Object, UFunction* Function, const FReplicationFlags& RepFlags) const;
 	virtual void InternalProcessRemoteFunction(class AActor* Actor, class UObject* SubObject, class UNetConnection* Connection, class UFunction* Function, void* Parms, FOutParmRec* OutParms, FFrame* Stack, bool bIsServer) override;
 	virtual int32 ServerReplicateActors(float DeltaSeconds) override;
@@ -304,6 +338,9 @@ public:
 	/** Called when the player controller associated with a connection is changed (either at the end of initial connection handshake, or if changed after successfully connected). */
 	void GameServerAssignPlayerController(UNetConnection* ChildGameServerConnection, UNetConnection* ParentGameServerConnection, APlayerController* PlayerController);
 
+	/** Return a pointer to the shared NetGUID cache used by backend network connections. */
+	TSharedPtr<FNetGUIDCache> GetSharedBackendNetGuidCache();
+
 private:
 
 	friend class UProxyListenerNotify;
@@ -317,9 +354,6 @@ private:
 
 	/** Set all of the configuration options that disable executing actor functionality or game specific code. */
 	void DisableActorLogicAndGameCode();
-
-	/** Return true if an actor is explicitly spawned by the proxy. */
-	bool IsProxySpawned(AActor* Actor) const;
 	
 	/** Send any split join requests that have been queued up to the game server. */
 	void FlushSplitJoinRequests(FGameServerConnectionState* GameServerConnectionState);
@@ -328,7 +362,7 @@ private:
 	int32 GetGameServerWithPlayer(class ULocalPlayer* Player);
 
 	/** Disassociate a proxy connection with a game server as a primary game server. */
-	void DeregisterPrimaryGameServer(UNetConnection* ProxyConnection);
+	void DeregisterPrimaryGameServer(UNetConnection* ProxyConnection, ULocalPlayer* PreviousPrimaryPlayer);
 
 	/** Associate a proxy connection with a primary game server. */
 	void RegisterPrimaryGameServer(UNetConnection* ProxyConnection, class ULocalPlayer* PrimaryPlayer, int32 GameServerConnectionStateIndex);
@@ -338,23 +372,14 @@ private:
 
 	/** Detach an existing player, player controller and connection. */
 	void DetachPlayerController(APlayerController* PlayerController);
-	
-	/** Detach an existing player, player controller and connection and destroy the controller. */
-	void DetachAndDestroyPlayerController(APlayerController* PlayerController);
-
-	/** Create and destroy a proxy player controller. Returns null if failed to spawn the controller. */
-	APlayerController* SpawnPlayerController();
 
-	/** Destroy a player controller spawned by the proxy. */
-	void DestroyPlayerController(APlayerController* PlayerController);
-	
 	/** Return a proxy connection from a primary player (a player on a primary game server). */
-	static UNetConnection* GetProxyConnectionFromPrimaryPlayer(UPlayer* Player);
+	UNetConnection* GetProxyConnectionFromPrimaryPlayer(UPlayer* Player);
 
 	/** Return a proxy player controller given a primary game server controller. */
 	APlayerController* GetProxyControllerFromPrimaryGameServerController(APlayerController* GameServerPlayerController);
 	
-	/** Associate a request to connect to a game server with a route from a proxy connection to game server. */
+	/** Associate a request to connect to a game server (client handshake id)  with a route from a proxy connection to game server. */
 	UPROPERTY()
 	TMap<uint32, FMultiServerProxyInternalConnectionRoute> ClientHandshakeIdToInternalRoute;
 
@@ -362,10 +387,17 @@ private:
 	UPROPERTY()
 	TMap<TObjectPtr<class ULocalPlayer>, int32> ProxyPlayerToPrimaryGameServer;
 
+	UPROPERTY()
+	TMap<TObjectPtr<UNetConnection>, TObjectPtr<class ULocalPlayer>> ProxyConnectionToPrimaryPlayer;
+
 	/** Map a player to it's primary game server player controller. */
 	UPROPERTY()
 	TMap<TObjectPtr<class ULocalPlayer>, TObjectPtr<APlayerController>> ProxyPlayerToPrimaryGameServerPlayerController;
 
+	/** Map a player to it's proxy connection. */
+	UPROPERTY()
+	TMap<TObjectPtr<class ULocalPlayer>, TObjectPtr<UNetConnection>> ProxyPlayerToProxyConnection;
+
 	/** The state of each incoming proxy connection. */
 	TMap<TObjectPtr<UNetConnection>, EProxyConnectionState> ProxyConnectionState;
 
@@ -376,11 +408,10 @@ private:
 	/** Proxy listener handshake logic. */
 	UPROPERTY()
 	TObjectPtr<UProxyListenerNotify> ProxyListenerNotify;
-
-	/** A set of actors that have been spawned by the proxy. */
-	UPROPERTY()
-	TSet<TObjectPtr<AActor>> ProxySpawnedActors;
 	
+	/** A single NetGUID cache used by all netdrivers used to connect to backend servers. */
+	TSharedPtr<FNetGUIDCache> SharedBackendNetGuidCache;
+
 	int32 NextGameServerClientId = 0;
 	
 	uint32 NextClientHandshakeId = 123;
diff --git a/Engine/Source/Runtime/CoreUObject/Private/UObject/CoreNet.cpp b/Engine/Source/Runtime/CoreUObject/Private/UObject/CoreNet.cpp
--- a/Engine/Source/Runtime/CoreUObject/Private/UObject/CoreNet.cpp
+++ b/Engine/Source/Runtime/CoreUObject/Private/UObject/CoreNet.cpp
@@ -686,6 +686,8 @@ const TCHAR* LexToString(const EChannelCloseReason Value)
 		return TEXT("Relevancy");
 	case EChannelCloseReason::TearOff:
 		return TEXT("TearOff");
+	case EChannelCloseReason::Migrated:
+		return TEXT("Migrated");
 	}
 
 	return TEXT("Unknown");
diff --git a/Engine/Source/Runtime/CoreUObject/Private/UObject/RemoteObject.cpp b/Engine/Source/Runtime/CoreUObject/Private/UObject/RemoteObject.cpp
--- a/Engine/Source/Runtime/CoreUObject/Private/UObject/RemoteObject.cpp
+++ b/Engine/Source/Runtime/CoreUObject/Private/UObject/RemoteObject.cpp
@@ -19,6 +19,7 @@
 #include "Misc/CommandLine.h"
 #include "Misc/CoreDelegates.h"
 #include "String/LexFromString.h"
+#include "Hash/CityHash.h"
 #include <atomic>
 
 DEFINE_LOG_CATEGORY(LogRemoteObject);
@@ -673,7 +674,7 @@ namespace UE::CoreUObject::Private
 #define UE_WITH_REMOTE_ASSET_ID 1 // set this to 0 to disable remote asset IDs
 #endif
 
-FRemoteObjectId FRemoteObjectId::Generate(UObjectBase* InObject, EInternalObjectFlags InInitialFlags /*= EInternalObjectFlags::None*/)
+FRemoteObjectId FRemoteObjectId::Generate(UObjectBase* InObject, const TCHAR* InName, EInternalObjectFlags InInitialFlags /*= EInternalObjectFlags::None*/)
 {
 	using namespace UE::RemoteObject::Private;
 	using namespace UE::CoreUObject::Private;
@@ -719,13 +720,38 @@ FRemoteObjectId FRemoteObjectId::Generate(UObjectBase* InObject, EInternalObject
 #endif // UE_WITH_REMOTE_ASSET_ID
 
 	FRemoteObjectId Result;
+
 	if (!bIsAsset)
 	{
 		Result = FRemoteObjectId(RemoteServerId, RemoteObjectSerialNumber.fetch_add(1));
+		UE_LOG(LogRemoteObject, Verbose, TEXT("FRemoteObjectId::Generate: Type=Dynamic Object=%s RemoteId=%s"), InObject ? *InObject->GetFName().ToString() : TEXT("None"), *Result.ToString());
 	}
 	else
 	{
-		Result = FRemoteObjectId(FRemoteServerId(ERemoteServerIdConstants::Asset), AssetObjectSerialNumber.fetch_add(1));
+		// An asset must have the same remote object id on each game server. In order to achieve this the object's full path is converted into a 53-bit hash (the lower
+		// 53 bits of a 64-bit hash) and used as the remote id.
+
+		FString ObjectOuterPath = (InObject && InObject->GetOuter()) ? InObject->GetOuter()->GetPathName() : TEXT("");
+		FString ObjectName(InName);
+
+		FStringBuilderBase ObjectPathBuilder;
+		if (ObjectOuterPath.Len() > 0)
+		{
+			ObjectPathBuilder << ObjectOuterPath;
+			ObjectPathBuilder << SUBOBJECT_DELIMITER_CHAR;
+		}
+		ObjectPathBuilder << ObjectName;
+
+		const uint64 ObjectPathHash64 = CityHash64(
+			reinterpret_cast<const char*>(ObjectPathBuilder.GetData()), 
+			ObjectPathBuilder.Len() * sizeof(FStringBuilderBase::ElementType)
+		);
+
+		const uint64 MaskLowBits = (1ULL << REMOTE_OBJECT_SERIAL_NUMBER_BIT_SIZE) - 1;
+		const uint64 ObjectPathHash53 = ObjectPathHash64 & MaskLowBits;
+
+		Result = FRemoteObjectId(FRemoteServerId(ERemoteServerIdConstants::Asset), ObjectPathHash53);
+		UE_LOG(LogRemoteObject, Verbose, TEXT("FRemoteObjectId::Generate: Type=Asset Object=%s RemoteId=%s Base64BitHash=%llu"), ObjectPathBuilder.ToString(), *Result.ToString(), ObjectPathHash64);
 	}
 
 	return Result;
diff --git a/Engine/Source/Runtime/CoreUObject/Private/UObject/UObjectArray.cpp b/Engine/Source/Runtime/CoreUObject/Private/UObject/UObjectArray.cpp
--- a/Engine/Source/Runtime/CoreUObject/Private/UObject/UObjectArray.cpp
+++ b/Engine/Source/Runtime/CoreUObject/Private/UObject/UObjectArray.cpp
@@ -245,7 +245,7 @@ void FUObjectArray::AllocateUObjectIndex(UObjectBase* Object, EInternalObjectFla
 #if UE_WITH_REMOTE_OBJECT_HANDLE
 	if (!RemoteId.IsValid())
 	{
-		RemoteId = FRemoteObjectId::Generate(Object, InitialFlags);
+		RemoteId = FRemoteObjectId::Generate(Object, *Object->GetFName().ToString(), InitialFlags);
 	}
 #endif
 
diff --git a/Engine/Source/Runtime/CoreUObject/Private/UObject/UObjectBase.cpp b/Engine/Source/Runtime/CoreUObject/Private/UObject/UObjectBase.cpp
--- a/Engine/Source/Runtime/CoreUObject/Private/UObject/UObjectBase.cpp
+++ b/Engine/Source/Runtime/CoreUObject/Private/UObject/UObjectBase.cpp
@@ -543,7 +543,7 @@ void UObjectBase::Register(class UClass* (*StaticClassFn)(), const TCHAR* Packag
 	FPendingRegistrant* PendingRegistration = new FPendingRegistrant(this);
 	PendingRegistrants.Add(this, FPendingRegistrantInfo(StaticClassFn, InName, PackageName
 #if UE_WITH_REMOTE_OBJECT_HANDLE
-		, FRemoteObjectId::Generate(this)
+		, FRemoteObjectId::Generate(this, InName)
 #endif
 	));
 
diff --git a/Engine/Source/Runtime/CoreUObject/Public/UObject/CoreNetTypes.h b/Engine/Source/Runtime/CoreUObject/Public/UObject/CoreNetTypes.h
--- a/Engine/Source/Runtime/CoreUObject/Public/UObject/CoreNetTypes.h
+++ b/Engine/Source/Runtime/CoreUObject/Public/UObject/CoreNetTypes.h
@@ -52,6 +52,7 @@ enum class EChannelCloseReason : uint8
 	LevelUnloaded,
 	Relevancy,
 	TearOff,
+	Migrated,
 	/* reserved */
 	MAX	= 15		// this value is used for serialization, modifying it may require a network version change
 };
diff --git a/Engine/Source/Runtime/CoreUObject/Public/UObject/RemoteObjectTypes.h b/Engine/Source/Runtime/CoreUObject/Public/UObject/RemoteObjectTypes.h
--- a/Engine/Source/Runtime/CoreUObject/Public/UObject/RemoteObjectTypes.h
+++ b/Engine/Source/Runtime/CoreUObject/Public/UObject/RemoteObjectTypes.h
@@ -12,11 +12,17 @@ class UObjectBase;
 class UObject;
 class FArchive;
 
+constexpr int32 REMOTE_OBJECT_SERVER_ID_BIT_SIZE = 10;
+constexpr int32 REMOTE_OBJECT_SERIAL_NUMBER_BIT_SIZE = 53;
+
+constexpr uint64 MAX_REMOTE_OBJECT_SERVER_ID = ((uint64_t)1 << REMOTE_OBJECT_SERVER_ID_BIT_SIZE) - 1;
+constexpr uint64 MAX_REMOTE_OBJECT_SERIAL_NUMBER = ((uint64_t)1 << REMOTE_OBJECT_SERIAL_NUMBER_BIT_SIZE) - 1;
+
 enum class ERemoteServerIdConstants : uint32
 {
 	Invalid = 0,
 	FirstValid,
-	Max = (1 << 10) - 1,
+	Max = MAX_REMOTE_OBJECT_SERVER_ID,
 
 	Database = Max,
 	Asset = Max - 1,
@@ -103,8 +109,9 @@ private:
 	{
 		struct
 		{
-			uint64 SerialNumber : 54;
-			uint64 ServerId : 10;
+			uint64 Reserved : 1;	/* Reserved for static flag in NetGUID */
+			uint64 SerialNumber : REMOTE_OBJECT_SERIAL_NUMBER_BIT_SIZE;
+			uint64 ServerId : REMOTE_OBJECT_SERVER_ID_BIT_SIZE;
 		};
 		uint64 Id = 0;
 	};
@@ -113,10 +120,20 @@ public:
 
 	FRemoteObjectId() = default;
 
+	static FRemoteObjectId CreateFromInt(uint64 InRawId)
+	{
+		FRemoteObjectId RemoteObjectId;
+		RemoteObjectId.Id = InRawId;
+		return RemoteObjectId;
+	}
+
 	FRemoteObjectId(FRemoteServerId InServerId, uint64 InSerialNumber)
 		: SerialNumber(InSerialNumber)
 		, ServerId(InServerId.Id)
 	{
+		// Make sure the arguments will fit into the different bitfields.
+		ensure((uint64_t)SerialNumber <= MAX_REMOTE_OBJECT_SERIAL_NUMBER);
+		ensure((uint64_t)InServerId.Id <= MAX_REMOTE_OBJECT_SERVER_ID);
 	}
 
 	explicit COREUOBJECT_API FRemoteObjectId(const UObjectBase* InObject);
@@ -172,7 +189,7 @@ public:
 
 	COREUOBJECT_API FString ToString() const;
 
-	COREUOBJECT_API static FRemoteObjectId Generate(UObjectBase* InObject, EInternalObjectFlags InInitialFlags = EInternalObjectFlags::None);
+	COREUOBJECT_API static FRemoteObjectId Generate(UObjectBase* InObject, const TCHAR* InName, EInternalObjectFlags InInitialFlags = EInternalObjectFlags::None);
 
 	friend COREUOBJECT_API FArchive& operator<<(FArchive& Ar, FRemoteObjectId& Id);
 };
diff --git a/Engine/Source/Runtime/Engine/Classes/Engine/ActorChannel.h b/Engine/Source/Runtime/Engine/Classes/Engine/ActorChannel.h
--- a/Engine/Source/Runtime/Engine/Classes/Engine/ActorChannel.h
+++ b/Engine/Source/Runtime/Engine/Classes/Engine/ActorChannel.h
@@ -434,7 +434,7 @@ protected:
 	/** Unmap all references to this object, so that if later we receive this object again, we can remap the original references */
 	ENGINE_API void MoveMappedObjectToUnmapped(const UObject* Object);
 
-	ENGINE_API void DestroyActorAndComponents();
+	ENGINE_API void DestroyActorAndComponents(EChannelCloseReason Reason = EChannelCloseReason::Destroyed);
 
 	ENGINE_API virtual bool CleanUp( const bool bForDestroy, EChannelCloseReason CloseReason ) override;
 
diff --git a/Engine/Source/Runtime/Engine/Classes/Engine/DemoNetDriver.h b/Engine/Source/Runtime/Engine/Classes/Engine/DemoNetDriver.h
--- a/Engine/Source/Runtime/Engine/Classes/Engine/DemoNetDriver.h
+++ b/Engine/Source/Runtime/Engine/Classes/Engine/DemoNetDriver.h
@@ -587,7 +587,7 @@ public:
 	/** Called when a "go to time" operation is completed. */
 	ENGINE_API void NotifyGotoTimeFinished(bool bWasSuccessful);
 
-	ENGINE_API virtual void NotifyActorDestroyed(AActor* ThisActor, bool IsSeamlessTravel=false) override;
+	ENGINE_API virtual void NotifyActorDestroyed(AActor* ThisActor, bool IsSeamlessTravel=false, EChannelCloseReason CloseReason = EChannelCloseReason::Destroyed) override;
 	ENGINE_API virtual void NotifyActorLevelUnloaded(AActor* Actor) override;
 	ENGINE_API virtual void NotifyStreamingLevelUnload(ULevel* InLevel) override;
 
diff --git a/Engine/Source/Runtime/Engine/Classes/Engine/NetDriver.h b/Engine/Source/Runtime/Engine/Classes/Engine/NetDriver.h
--- a/Engine/Source/Runtime/Engine/Classes/Engine/NetDriver.h
+++ b/Engine/Source/Runtime/Engine/Classes/Engine/NetDriver.h
@@ -1754,7 +1754,7 @@ public:
 	ENGINE_API virtual void AddNetworkActor(AActor* Actor);
 
 	/** Called when a spawned actor is destroyed. */
-	ENGINE_API virtual void NotifyActorDestroyed( AActor* Actor, bool IsSeamlessTravel=false );
+	ENGINE_API virtual void NotifyActorDestroyed(AActor* Actor, bool IsSeamlessTravel = false, EChannelCloseReason CloseReason = EChannelCloseReason::Destroyed);
 
 	void NotifySubObjectDestroyed(UObject* SubObject);
 
@@ -1765,7 +1765,7 @@ public:
 	/** Called when an actor is renamed. */
 	ENGINE_API virtual void NotifyActorRenamed(AActor* Actor, UObject* PreviousOuter, FName PreviousName);
 
-	ENGINE_API void RemoveNetworkActor(AActor* Actor);
+	ENGINE_API virtual void RemoveNetworkActor(AActor* Actor);
 
 	/** Called when an authoritative actor wants to delete a replicated subobject on the clients it was already replicated to */
 	void DeleteSubObjectOnClients(AActor* Actor, UObject* SubObject);
@@ -2024,7 +2024,7 @@ public:
 	ENGINE_API virtual bool ShouldCallRemoteFunction(UObject* Object, UFunction* Function, const FReplicationFlags& RepFlags) const;
 
 	/** Returns true if clients should destroy the actor when the channel is closed. */
-	ENGINE_API virtual bool ShouldClientDestroyActor(AActor* Actor) const;
+	ENGINE_API virtual bool ShouldClientDestroyActor(AActor* Actor, EChannelCloseReason CloseReason = EChannelCloseReason::Destroyed) const;
 
 	/** Called when an actor channel is remotely opened for an actor. */
 	ENGINE_API virtual void NotifyActorChannelOpen(UActorChannel* Channel, AActor* Actor);
diff --git a/Engine/Source/Runtime/Engine/Classes/Engine/PackageMapClient.h b/Engine/Source/Runtime/Engine/Classes/Engine/PackageMapClient.h
--- a/Engine/Source/Runtime/Engine/Classes/Engine/PackageMapClient.h
+++ b/Engine/Source/Runtime/Engine/Classes/Engine/PackageMapClient.h
@@ -194,7 +194,8 @@ namespace UE::Net::Private
 class FNetGUIDCache
 {
 public:
-	ENGINE_API FNetGUIDCache( UNetDriver * InDriver );
+	ENGINE_API FNetGUIDCache(UNetDriver* InDriver, uint64 NetworkGuidSeed = 0);
+	ENGINE_API virtual ~FNetGUIDCache() {};
 
 	PRAGMA_DISABLE_DEPRECATION_WARNINGS
 	FNetGUIDCache(FNetGUIDCache&&) = default;
@@ -220,12 +221,12 @@ public:
 	ENGINE_API void			CleanReferences();
 	ENGINE_API bool			SupportsObject( const UObject* Object, const TWeakObjectPtr<UObject>* WeakObjectPtr=nullptr /** Optional: pass in existing weakptr to prevent this function from constructing one internally */ ) const;
 	ENGINE_API bool			IsDynamicObject( const UObject* Object );
-	ENGINE_API bool			IsNetGUIDAuthority() const;
+	ENGINE_API virtual bool	IsNetGUIDAuthority() const;
 	ENGINE_API FNetworkGUID	GetOrAssignNetGUID( UObject* Object, const TWeakObjectPtr<UObject>* WeakObjectPtr=nullptr /** Optional: pass in existing weakptr to prevent this function from constructing one internally */ );
 	ENGINE_API FNetworkGUID	GetNetGUID( const UObject* Object ) const;
 	ENGINE_API FNetworkGUID	GetOuterNetGUID( const FNetworkGUID& NetGUID ) const;
-	ENGINE_API FNetworkGUID	AssignNewNetGUID_Server( UObject* Object );
-	ENGINE_API FNetworkGUID	AssignNewNetGUIDFromPath_Server( const FString& PathName, UObject* ObjOuter, UClass* ObjClass );
+	ENGINE_API virtual FNetworkGUID	AssignNewNetGUID_Server( UObject* Object );
+	ENGINE_API virtual FNetworkGUID	AssignNewNetGUIDFromPath_Server( const FString& PathName, UObject* ObjOuter, UClass* ObjClass );
 	ENGINE_API void			RegisterNetGUID_Internal( const FNetworkGUID& NetGUID, const FNetGuidCacheObject& CacheObject );
 	ENGINE_API void			RegisterNetGUID_Server( const FNetworkGUID& NetGUID, UObject* Object );
 	ENGINE_API void			RegisterNetGUID_Client( const FNetworkGUID& NetGUID, const UObject* Object );
@@ -249,6 +250,8 @@ public:
 	ENGINE_API bool			CanClientLoadObject( const UObject* Object, const FNetworkGUID& NetGUID ) const;
 	ENGINE_API FString			Describe(const FNetworkGUID& NetGUID) const;
 
+	ENGINE_API bool 		RemoveNetGUID(const UObject* Object);
+
 	ENGINE_API void			AsyncPackageCallback(const FName& PackageName, UPackage * Package, EAsyncLoadingResult::Type Result);
 
 	ENGINE_API void			ResetCacheForDemo();
@@ -292,6 +295,11 @@ public:
 
 	bool											IsExportingNetGUIDBunch;
 
+#if UE_WITH_REMOTE_OBJECT_HANDLE
+	static FNetworkGUID GetNetworkGUIDFromRemoteObjectId(const FRemoteObjectId RemoteObjectId, const bool bIsStatic);
+	static FRemoteObjectId GetRemoteObjectIdFromNetworkGUID(const FNetworkGUID NetGuid);
+#endif
+
 private:
 
 	friend class UPackageMapClient;
diff --git a/Engine/Source/Runtime/Engine/Classes/GameFramework/CharacterMovementComponent.h b/Engine/Source/Runtime/Engine/Classes/GameFramework/CharacterMovementComponent.h
--- a/Engine/Source/Runtime/Engine/Classes/GameFramework/CharacterMovementComponent.h
+++ b/Engine/Source/Runtime/Engine/Classes/GameFramework/CharacterMovementComponent.h
@@ -3256,5 +3256,7 @@ public:
 	/** Returns base time delta to use for a ServerMove, default calculation (no time discrepancy resolution) */
 	ENGINE_API float GetBaseServerMoveDeltaTime(float ClientTimeStamp, float ActorTimeDilation) const;
 
+	ENGINE_API virtual void Serialize(FArchive& Archive);
+
 };
 
diff --git a/Engine/Source/Runtime/Engine/Classes/GameFramework/CharacterMovementReplication.h b/Engine/Source/Runtime/Engine/Classes/GameFramework/CharacterMovementReplication.h
--- a/Engine/Source/Runtime/Engine/Classes/GameFramework/CharacterMovementReplication.h
+++ b/Engine/Source/Runtime/Engine/Classes/GameFramework/CharacterMovementReplication.h
@@ -8,6 +8,7 @@
 #include "Serialization/BitWriter.h"
 #include "Containers/BitArray.h"
 #include "Net/Core/NetToken/NetTokenExportContext.h"
+#include "Components/PrimitiveComponent.h"
 #if UE_WITH_IRIS
 #include "Iris/Serialization/IrisObjectReferencePackageMap.h"
 #endif
@@ -285,6 +286,25 @@ public:
 	bool bBaseRelativePosition;
 	bool bBaseRelativeVelocity;
 	uint8 MovementMode;
+
+	void Serialize(FArchive& Archive)
+	{
+		if (Archive.IsMigratingRemoteObjects())
+		{
+			Archive << TimeStamp;
+			Archive << DeltaTime;
+			Archive << NewLoc;
+			Archive << NewVel;
+			Archive << NewRot;
+			Archive << GravityDirection;
+			Archive << NewBase;
+			Archive << NewBaseBoneName;
+			Archive.SerializeBits(&bAckGoodMove, 1);
+			Archive.SerializeBits(&bBaseRelativePosition, 1);
+			Archive.SerializeBits(&bBaseRelativeVelocity, 1);
+			Archive << MovementMode;
+		}
+	}
 };
 
 
diff --git a/Engine/Source/Runtime/Engine/Classes/GameFramework/PlayerController.h b/Engine/Source/Runtime/Engine/Classes/GameFramework/PlayerController.h
--- a/Engine/Source/Runtime/Engine/Classes/GameFramework/PlayerController.h
+++ b/Engine/Source/Runtime/Engine/Classes/GameFramework/PlayerController.h
@@ -1848,9 +1848,7 @@ public:
 	ENGINE_API virtual void DisableInput(class APlayerController* PlayerController) override;
 
 	ENGINE_API virtual void Serialize(FArchive& Ar) override;
-#if UE_WITH_REMOTE_OBJECT_HANDLE
-	ENGINE_API virtual void Tick(float DeltaTime) override;
-#endif
+
 protected:
 	ENGINE_API virtual void BeginPlay() override;
 	//~ End AActor Interface
@@ -2315,9 +2313,6 @@ public:
 	ENGINE_API void SetClientHandshakeId(uint32 InClientHandshakeId);
 
 private:
-	
-	// This is a server-specific variable so it should not be migrated between servers and remain non-UPROPERTY().
-	bool bSendClientRestartRPC = false;
 
 	/** The player id associated with this player controller. */
 	UPROPERTY()
diff --git a/Engine/Source/Runtime/Engine/Classes/Interfaces/NetworkPredictionInterface.h b/Engine/Source/Runtime/Engine/Classes/Interfaces/NetworkPredictionInterface.h
--- a/Engine/Source/Runtime/Engine/Classes/Interfaces/NetworkPredictionInterface.h
+++ b/Engine/Source/Runtime/Engine/Classes/Interfaces/NetworkPredictionInterface.h
@@ -127,5 +127,27 @@ public:
 
 	/** Set to true if last received move request  is bad and needs correction */
 	uint8 bLastRequestNeedsForcedUpdates : 1;
+
+	virtual void Serialize(FArchive& Archive)
+	{
+		if (Archive.IsMigratingRemoteObjects())
+		{
+			Archive << ServerTimeStamp;
+			Archive << ServerTimeBeginningForcedUpdates;
+			Archive << ServerTimeLastForcedUpdate;
+			
+			bool bTempTriggeringForcedUpdates = bTriggeringForcedUpdates;
+			Archive << bTempTriggeringForcedUpdates;
+			bTriggeringForcedUpdates = bTempTriggeringForcedUpdates;
+
+			bool bTempForcedUpdateDurationExceeded = bForcedUpdateDurationExceeded;
+			Archive << bTempForcedUpdateDurationExceeded;
+			bForcedUpdateDurationExceeded = bTempForcedUpdateDurationExceeded;
+
+			bool bTempLastRequestNeedsForcedUpdates = bLastRequestNeedsForcedUpdates;
+			Archive << bTempLastRequestNeedsForcedUpdates;
+			bLastRequestNeedsForcedUpdates = bTempLastRequestNeedsForcedUpdates;
+		}
+	}
 };
 
diff --git a/Engine/Source/Runtime/Engine/Private/Actor.cpp b/Engine/Source/Runtime/Engine/Private/Actor.cpp
--- a/Engine/Source/Runtime/Engine/Private/Actor.cpp
+++ b/Engine/Source/Runtime/Engine/Private/Actor.cpp
@@ -63,6 +63,7 @@
 #include "Engine/SCS_Node.h"
 #include "HAL/LowLevelMemStats.h"
 #include "ProfilingDebugging/AssetMetadataTrace.h"
+#include "Engine/PackageMapClient.h"
 
 #if UE_WITH_REMOTE_OBJECT_HANDLE
 #include "UObject/RemoteExecutor.h"
@@ -1312,7 +1313,7 @@ void AActor::PostMigrate(const FUObjectMigrationContext& MigrationContext)
 					{
 						if (Driver.NetDriver != nullptr && Driver.NetDriver->ShouldReplicateActor(this))
 						{
-							Driver.NetDriver->NotifyActorDestroyed(this);
+							Driver.NetDriver->NotifyActorDestroyed(this, false, EChannelCloseReason::Migrated);
 						}
 					}
 				}
@@ -1334,6 +1335,18 @@ void AActor::PostMigrate(const FUObjectMigrationContext& MigrationContext)
 			// Unregister the actor's tick function
 			RegisterAllActorTickFunctions(/*bRegisterTickFunctions*/ false, /*bIncludeComponents*/ true);
 		}
+
+		if (MigrationSendType == EObjectMigrationSendType::ReturnBorrowed)
+		{
+			if (GetWorld()->GetNetDriver()->GuidCache->RemoveNetGUID(this))
+			{
+				UE_LOG(LogActor, Verbose, TEXT("Deregistering loaned object %s from the NetGUID cache."), *this->GetName());
+			}
+			else
+			{
+				UE_LOG(LogActor, Warning, TEXT("Loaned object %s isn't registered in the NetGUID cache."), *GetName());
+			}
+		}
 	}
 	else if (MigrationContext.MigrationSide == EObjectMigrationSide::Receive)
 	{
@@ -1413,6 +1426,30 @@ void AActor::PostMigrate(const FUObjectMigrationContext& MigrationContext)
 				GetLevel()->TryAddActorToList(this, bOnlyAddIfUnique);
 			}
 		}
+
+		if (MigrationRecvType == EObjectMigrationRecvType::Borrowed)
+		{
+			TSharedPtr<FNetGUIDCache>& GuidCache = GetWorld()->GetNetDriver()->GuidCache;
+			const bool bIsStatic = !GuidCache->IsDynamicObject(this);
+			FRemoteObjectId RemoteObjectId = FRemoteObjectId(this);
+			FNetworkGUID NetGUID = FNetGUIDCache::GetNetworkGUIDFromRemoteObjectId(RemoteObjectId, bIsStatic);
+
+			UE_LOG(LogActor, Verbose, TEXT("Registering borrowed object %s (NetGUID=%s RemoteObjectId=%s) in the NetGUID cache."),
+				   *this->GetName(),
+				   *NetGUID.ToString(),
+				   *RemoteObjectId.ToString());
+
+			if (!GuidCache->GetNetGUID(this).IsValid())
+			{
+				GuidCache->RegisterNetGUID_Server(NetGUID, this);
+			}
+			else
+			{
+				UE_LOG(LogActor, Warning, TEXT("Borrowed object %s is already registered in the NetGUID cache."), *GetName());
+			}
+		}
+
+		ActorHasBegunPlay = EActorBeginPlayState::HasBegunPlay;
 	}
 }
 
diff --git a/Engine/Source/Runtime/Engine/Private/Components/CharacterMovementComponent.cpp b/Engine/Source/Runtime/Engine/Private/Components/CharacterMovementComponent.cpp
--- a/Engine/Source/Runtime/Engine/Private/Components/CharacterMovementComponent.cpp
+++ b/Engine/Source/Runtime/Engine/Private/Components/CharacterMovementComponent.cpp
@@ -1571,6 +1571,32 @@ void UCharacterMovementComponent::Serialize(FArchive& Archive)
 		// Since UpdatedComponent is transient, we'll have to wait until we're registered.
 		bNeedsSweepWhileWalkingUpdate = true;
 	}
+	
+	if (Archive.IsMigratingRemoteObjects())
+	{
+		bool bHasServerPredictionData = false;
+
+		if (Archive.IsSaving())
+		{
+			bHasServerPredictionData = (ServerPredictionData != nullptr);
+			Archive << bHasServerPredictionData;
+		}
+		else if (Archive.IsLoading())
+		{
+			Archive << bHasServerPredictionData;
+			if (bHasServerPredictionData && ServerPredictionData == nullptr)
+			{
+				ServerPredictionData = new FNetworkPredictionData_Server_Character(*this);
+			}
+		}
+
+		if (bHasServerPredictionData)
+		{
+			ensure(ServerPredictionData);
+
+			ServerPredictionData->Serialize(Archive);
+		}
+	}
 }
 
 void UCharacterMovementComponent::AddReferencedObjects(UObject* InThis, FReferenceCollector& Collector)
@@ -12316,6 +12342,34 @@ float FNetworkPredictionData_Server_Character::GetBaseServerMoveDeltaTime(float
 	return DeltaTime;
 }
 
+void FNetworkPredictionData_Server_Character::Serialize(FArchive& Archive)
+{
+	if (Archive.IsMigratingRemoteObjects())
+	{
+		FNetworkPredictionData_Server::Serialize(Archive);
+
+		PendingAdjustment.Serialize(Archive);
+
+		Archive << CurrentClientTimeStamp;
+		Archive << LastReceivedClientTimeStamp;
+		Archive << ServerAccumulatedClientTimeStamp;
+		Archive << LastUpdateTime;
+		Archive << ServerTimeStampLastServerMove;
+		Archive << MaxMoveDeltaTime;
+		
+		bool bTempForceClientUpdate = bForceClientUpdate;
+		Archive.SerializeBits(&bTempForceClientUpdate, 1);
+		bForceClientUpdate = bTempForceClientUpdate;
+
+		Archive << LifetimeRawTimeDiscrepancy;
+		Archive << TimeDiscrepancy;
+		Archive.SerializeBits(&bResolvingTimeDiscrepancy, 1);
+		Archive << TimeDiscrepancyResolutionMoveDeltaOverride;
+		Archive << TimeDiscrepancyAccumulatedClientDeltasSinceLastServerTick;
+		Archive << WorldCreationTime;
+	}
+};
+
 
 FSavedMove_Character::FSavedMove_Character()
 {
diff --git a/Engine/Source/Runtime/Engine/Private/DataChannel.cpp b/Engine/Source/Runtime/Engine/Private/DataChannel.cpp
--- a/Engine/Source/Runtime/Engine/Private/DataChannel.cpp
+++ b/Engine/Source/Runtime/Engine/Private/DataChannel.cpp
@@ -2513,7 +2513,7 @@ void UActorChannel::MoveMappedObjectToUnmapped(const UObject* Object)
 	}
 }
 
-void UActorChannel::DestroyActorAndComponents()
+void UActorChannel::DestroyActorAndComponents(EChannelCloseReason Reason)
 {
 	// Destroy any sub-objects we created
 	PRAGMA_DISABLE_DEPRECATION_WARNINGS
@@ -2686,11 +2686,11 @@ bool UActorChannel::CleanUp(const bool bForDestroy, EChannelCloseReason CloseRea
 			{
 				UE_LOG(LogNetTraffic, Verbose, TEXT("UActorChannel::CleanUp: Skipped Destroying NetStartupActor %s due to EChannelCloseReason::LevelUnloaded"), *Describe());
 			}
-			else if (!Actor->bNetTemporary && Actor->GetWorld() != NULL && !IsEngineExitRequested() && Connection->Driver->ShouldClientDestroyActor(Actor))
+			else if (!Actor->bNetTemporary && Actor->GetWorld() != NULL && !IsEngineExitRequested() && Connection->Driver->ShouldClientDestroyActor(Actor, CloseReason))
 			{
 				UE_LOG(LogNetDormancy, Verbose, TEXT("UActorChannel::CleanUp: Destroying Actor (ChIndex: %i): %s"), ChIndex, *GetFullNameSafe(Actor));
 
-				DestroyActorAndComponents();
+				DestroyActorAndComponents(CloseReason);
 			}
 		}
 	}
diff --git a/Engine/Source/Runtime/Engine/Private/DemoNetDriver.cpp b/Engine/Source/Runtime/Engine/Private/DemoNetDriver.cpp
--- a/Engine/Source/Runtime/Engine/Private/DemoNetDriver.cpp
+++ b/Engine/Source/Runtime/Engine/Private/DemoNetDriver.cpp
@@ -5113,7 +5113,7 @@ void UDemoNetDriver::OnActorPreDestroy(AActor* DestroyedActor)
 	}
 }
 
-void UDemoNetDriver::NotifyActorDestroyed(AActor* Actor, bool IsSeamlessTravel)
+void UDemoNetDriver::NotifyActorDestroyed(AActor* Actor, bool IsSeamlessTravel, EChannelCloseReason CloseReason)
 {
 	check(Actor != nullptr);
 
@@ -5169,7 +5169,7 @@ void UDemoNetDriver::NotifyActorDestroyed(AActor* Actor, bool IsSeamlessTravel)
 		}
 	}
 
-	Super::NotifyActorDestroyed(Actor, IsSeamlessTravel);
+	Super::NotifyActorDestroyed(Actor, IsSeamlessTravel, CloseReason);
 }
 
 void UDemoNetDriver::NotifyActorTornOff(AActor* Actor)
diff --git a/Engine/Source/Runtime/Engine/Private/NetDriver.cpp b/Engine/Source/Runtime/Engine/Private/NetDriver.cpp
--- a/Engine/Source/Runtime/Engine/Private/NetDriver.cpp
+++ b/Engine/Source/Runtime/Engine/Private/NetDriver.cpp
@@ -835,7 +835,7 @@ void UNetDriver::AddNetworkActor(AActor* Actor)
 	LLM_SCOPE_BYTAG(NetDriver);
 	ensureMsgf(Actor == nullptr || !(Actor->HasAnyFlags(RF_ClassDefaultObject | RF_ArchetypeObject)), TEXT("%s is a CDO or Archetype and should not be replicated."), *GetFullNameSafe(Actor));
 
-	UE_LOG(LogNetLifecycle, Log, TEXT("Registering for replication: Actor: %s"), *GetFullNameSafe(Actor));
+	UE_LOG(LogNetLifecycle, Log, TEXT("[%s] Registering for replication: Actor: %s"), *GetName(), *GetFullNameSafe(Actor));
 
 #if UE_WITH_IRIS
 	// We currently only want to do this on the server as we want to replicate the actor to the client
@@ -4064,7 +4064,7 @@ bool UNetDriver::CanCreateDestructionInfo() const
 	return !GIsReconstructingBlueprintInstances && !UE::Net::ShouldIgnoreStaticActorDestruction();
 }
 
-void UNetDriver::NotifyActorDestroyed(AActor* ThisActor, bool IsSeamlessTravel)
+void UNetDriver::NotifyActorDestroyed(AActor* ThisActor, bool IsSeamlessTravel, EChannelCloseReason CloseReason)
 {
 #if UE_WITH_IRIS
 	if (ReplicationSystem)
@@ -4133,7 +4133,7 @@ void UNetDriver::NotifyActorDestroyed(AActor* ThisActor, bool IsSeamlessTravel)
 			{
 				check(Channel->OpenedLocally);
 				Channel->bClearRecentActorRefs = false;
-				Channel->Close(EChannelCloseReason::Destroyed);
+				Channel->Close(CloseReason);
 			}
 			else
 			{
@@ -4174,7 +4174,7 @@ void UNetDriver::NotifySubObjectDestroyed(UObject* SubObject)
 void UNetDriver::RemoveNetworkActor(AActor* Actor)
 {
 	LLM_SCOPE_BYTAG(NetDriver);
-	UE_LOG(LogNetLifecycle, Log, TEXT("Deregistering for replication: Actor: %s"), *GetFullNameSafe(Actor));
+	UE_LOG(LogNetLifecycle, Log, TEXT("[%s] Deregistering for replication: Actor: %s"), *GetName(), *GetFullNameSafe(Actor));
 
 	// This is a hack: Sometimes world streaming will stream out levels even though some players can still see them
 	// Usually this is not a problem since the clients will remove the actors when it's told the level is gone. 
@@ -8010,7 +8010,7 @@ bool UNetDriver::ShouldCallRemoteFunction(UObject* Object, UFunction* Function,
 	return ((!IsServer() || RepFlags.bNetOwner) && !RepFlags.bIgnoreRPCs);
 }
 
-bool UNetDriver::ShouldClientDestroyActor(AActor* Actor) const
+bool UNetDriver::ShouldClientDestroyActor(AActor* Actor, EChannelCloseReason CloseReason) const
 {
 	return (Actor && !Actor->IsA(ALevelScriptActor::StaticClass()));
 }
@@ -8959,4 +8959,4 @@ FConsoleCommandWithArgsDelegate::CreateLambda([](const TArray< FString >& Args)
 	}
 }));
 
-} // end namespace UE::Net::Private
+} // end namespace UE::Net::Private
\ No newline at end of file
diff --git a/Engine/Source/Runtime/Engine/Private/PackageMapClient.cpp b/Engine/Source/Runtime/Engine/Private/PackageMapClient.cpp
--- a/Engine/Source/Runtime/Engine/Private/PackageMapClient.cpp
+++ b/Engine/Source/Runtime/Engine/Private/PackageMapClient.cpp
@@ -44,6 +44,10 @@ static const int GUID_PACKET_ACKED		= -1;
 
 CSV_DEFINE_CATEGORY(PackageMap, true);
 
+#if UE_WITH_REMOTE_OBJECT_HANDLE
+DEFINE_LOG_CATEGORY_STATIC(LogNetPackageMapRemoteObject, Warning, All);
+#endif
+
 /**
  * Don't allow infinite recursion of InternalLoadObject - an attacker could
  * send malicious packets that cause a stack overflow on the server.
@@ -510,6 +514,11 @@ bool UPackageMapClient::SerializeNewActor(FArchive& Ar, class UActorChannel *Cha
 	UObject *NewObj = Actor;
 	SerializeObject(Ar, AActor::StaticClass(), NewObj, &NetGUID);
 
+#if UE_WITH_REMOTE_OBJECT_HANDLE
+	UE_LOG(LogNetPackageMapRemoteObject, Log, TEXT("UPackageMapClient::SerializeNewActor: Actor=%s NetGUID=%s"),
+		   *GetNameSafe(NewObj), *NetGUID.ToString());
+#endif
+
 	if ( Ar.IsError() )
 	{
 		UE_LOG( LogNetPackageMap, Error, TEXT( "UPackageMapClient::SerializeNewActor: Ar.IsError after SerializeObject 1" ) );
@@ -1115,6 +1124,28 @@ FNetworkGUID UPackageMapClient::InternalLoadObject( FArchive & Ar, UObject *& Ob
 	{
 		Object = GetObjectFromNetGUID( NetGUID, GuidCache->IsExportingNetGUIDBunch );
 
+#if UE_WITH_REMOTE_OBJECT_HANDLE
+		UE_LOG(LogNetPackageMapRemoteObject, Log, TEXT("Looking up NetGUID %s in the local cache and found '%s'"), *NetGUID.ToString(), *GetNameSafe(Object));
+
+		if (Object == nullptr)
+		{
+			FRemoteObjectId RemoteObjectId = FNetGUIDCache::GetRemoteObjectIdFromNetworkGUID(NetGUID);
+			if (UE::RemoteObject::Handle::IsRemote(RemoteObjectId))
+			{
+				if (AutoRTFM::IsTransactional())
+				{
+					UE_LOG(LogNetPackageMapRemoteObject, Log, TEXT("Dereferencing a remote NetGUID %s (%s)."), *NetGUID.ToString(), *RemoteObjectId.ToString());
+					TObjectPtr<UObject> RemoteObjectPtr = TObjectPtr<UObject>(FObjectPtr(RemoteObjectId));
+					Object = RemoteObjectPtr.Get();
+				}
+				else
+				{
+					UE_LOG(LogNetPackageMapRemoteObject, Log, TEXT("Dereferencing a remote NetGUID %s (%s) that is outside a transaction."), *NetGUID.ToString(), *RemoteObjectId.ToString());
+				}
+			}
+		}
+#endif
+
 		UE_LOG(LogNetPackageMap, VeryVerbose, TEXT( "InternalLoadObject loaded %s from NetGUID <%s>" ), Object ? *Object->GetFullName() : TEXT( "NULL" ), *NetGUID.ToString() );
 	}
 
@@ -2902,6 +2933,21 @@ const TArray<FNetworkGUID>* FNetGUIDCache::FindUnmappedStablyNamedGuidsWithOuter
 	return nullptr;
 }
 
+#if UE_WITH_REMOTE_OBJECT_HANDLE
+FNetworkGUID FNetGUIDCache::GetNetworkGUIDFromRemoteObjectId(const FRemoteObjectId RemoteObjectId, const bool bIsStatic)
+{
+	// Ignore the first bit of FRemoteObjectId which is reserved for the NetGUID static flag.
+	return FNetworkGUID::CreateFromIndex(RemoteObjectId.GetIdNumber() >> 1, bIsStatic);
+}
+
+FRemoteObjectId FNetGUIDCache::GetRemoteObjectIdFromNetworkGUID(const FNetworkGUID NetGuid)
+{
+	// The NetGUID ObjectId is a 64-bit value with the first bit as the static flag. Since this value is ignored by
+	// FRemoteObjectId (it's a reserved bit) the value is masked away.
+	return FRemoteObjectId::CreateFromInt(NetGuid.ObjectId & ~1);
+}
+#endif
+
 void UPackageMapClient::AddUnmappedNetGUIDReference(FNetworkGUID UnmappedGUID)
 {
 	using namespace UE::Net::Private;
@@ -2953,7 +2999,7 @@ void UPackageMapClient::RemoveUnmappedNetGUIDReference(FNetworkGUID NetGUID)
 //	FNetGUIDCache
 //----------------------------------------------------------------------------------------
 
-FNetGUIDCache::FNetGUIDCache(UNetDriver* InDriver) 
+FNetGUIDCache::FNetGUIDCache(UNetDriver* InDriver, uint64 NetworkGuidSeed) 
 	: Driver(InDriver)
 	, NetworkChecksumMode(ENetworkChecksumMode::SaveAndUse)
 	, AsyncLoadMode(EAsyncLoadMode::UseCVar)
@@ -2962,8 +3008,6 @@ FNetGUIDCache::FNetGUIDCache(UNetDriver* InDriver)
 {
 	UniqueNetFieldExportGroupPathIndex = 0;
 
-	uint64 NetworkGuidSeed = 0;
-
 #if !UE_BUILD_SHIPPING
 	FParse::Value(FCommandLine::Get(), TEXT("NetworkGuidSeed="), NetworkGuidSeed);
 #endif
@@ -3244,7 +3288,18 @@ FNetworkGUID FNetGUIDCache::AssignNewNetGUID_Server( UObject* Object )
 	// Generate new NetGUID and assign it
 	const int32 IsStatic = IsDynamicObject( Object ) ? 0 : 1;
 
+#if UE_WITH_REMOTE_OBJECT_HANDLE
+	const FRemoteObjectId RemoteObjectId = TObjectPtr<UObject>(Object).GetRemoteId();
+	const FNetworkGUID NewNetGuid = GetNetworkGUIDFromRemoteObjectId(RemoteObjectId, IsStatic != 0);
+
+	UE_LOG(LogNetPackageMapRemoteObject, Log, TEXT("Assigning object %s to use NetGUID %s based on remote object id %s"),
+		   *Object->GetName(),
+		   *NewNetGuid.ToString(),
+		   *RemoteObjectId.ToString());
+
+#else
 	const FNetworkGUID NewNetGuid = FNetworkGUID::CreateFromIndex(++NetworkGuidIndex[IsStatic], IsStatic != 0);
+#endif
 
 	RegisterNetGUID_Server( NewNetGuid, Object );
 
@@ -3263,7 +3318,18 @@ FNetworkGUID FNetGUIDCache::AssignNewNetGUIDFromPath_Server( const FString& Path
 	FNetworkGUID OuterGUID = GetOrAssignNetGUID( ObjOuter );
 
 	// Generate new NetGUID and assign it
+#if UE_WITH_REMOTE_OBJECT_HANDLE
+	const FRemoteObjectId RemoteObjectId = TObjectPtr<UObject>(ObjOuter).GetRemoteId();
+	const FNetworkGUID NewNetGuid = FNetworkGUID::CreateFromIndex(RemoteObjectId.GetIdNumber(), true);
+
+	UE_LOG(LogNetPackageMapRemoteObject, Log, TEXT("Assigning object %s from path %s to use NetGUID %s based on remote object id %s"),
+		   *ObjOuter->GetName(),
+		   *PathName,
+		   *NewNetGuid.ToString(),
+		   *RemoteObjectId.ToString());
+#else
 	const FNetworkGUID NewNetGuid = FNetworkGUID::CreateFromIndex(++NetworkGuidIndex[1], true);
+#endif
 
 	uint32 NetworkChecksum = GetClassNetworkChecksum( ObjClass );
 
@@ -4118,6 +4184,23 @@ FString	FNetGUIDCache::Describe(const FNetworkGUID& NetGUID) const
 	return Desc;
 }
 
+bool FNetGUIDCache::RemoveNetGUID(const UObject* Object)
+{
+	FNetworkGUID* NetGUIDPtr = NetGUIDLookup.Find(Object);
+	if (NetGUIDPtr)
+	{
+		FNetworkGUID NetGUID = *NetGUIDPtr;
+
+		ensure(ObjectLookup.Contains(NetGUID));
+		ObjectLookup.Remove(NetGUID);
+		NetGUIDLookup.Remove(Object);
+
+		return true;
+	}
+
+	return false;
+}
+
 void FNetGUIDCache::GenerateFullNetGUIDPath_r( const FNetworkGUID& NetGUID, FString& FullPath ) const
 {
 	if ( !NetGUID.IsValid() )
diff --git a/Engine/Source/Runtime/Engine/Private/PlayerController.cpp b/Engine/Source/Runtime/Engine/Private/PlayerController.cpp
--- a/Engine/Source/Runtime/Engine/Private/PlayerController.cpp
+++ b/Engine/Source/Runtime/Engine/Private/PlayerController.cpp
@@ -5008,7 +5008,6 @@ void APlayerController::PostMigrate(const FUObjectMigrationContext& MigrationCon
 
 		if (bRegisterWithWorld)
 		{
-			bSendClientRestartRPC = true;
 			GetWorld()->AddController(this);
 		}
 	}
@@ -5124,29 +5123,6 @@ void APlayerController::Serialize(FArchive& Ar)
 #endif
 }
 
-#if UE_WITH_REMOTE_OBJECT_HANDLE
-void APlayerController::Tick(float DeltaTime)
-{
-	Super::Tick(DeltaTime);
-
-	if (bSendClientRestartRPC)
-	{
-		if (!GetPawn().IsRemote())
-		{
-			UE_LOG(LogPlayerController, Log, TEXT("[%s] APlayerController::Tick [%s]: Repossessing pawn %s after migration."), 
-				   *UE::RemoteObject::GetGlobalServerId().ToString(),
-				   *FRemoteObjectId(this).ToString(),
-				   *GetPawn()->GetName());
-
-			GetPawn()->OnDestroyed.Clear();
-			Possess(GetPawn());
-			
-			bSendClientRestartRPC = false;
-		}
-	}
-}
-#endif
-
 void APlayerController::SetTargetViewRotation(const FRotator& InRotation)
 {
 	if (UE::Gameplay::CVars::bIsPlayerControllerPushBased)
diff --git a/Engine/Source/Runtime/Landscape/Classes/LandscapeProxy.h b/Engine/Source/Runtime/Landscape/Classes/LandscapeProxy.h
--- a/Engine/Source/Runtime/Landscape/Classes/LandscapeProxy.h
+++ b/Engine/Source/Runtime/Landscape/Classes/LandscapeProxy.h
@@ -891,7 +891,7 @@ PRAGMA_ENABLE_DEPRECATION_WARNINGS
 	bool bIsRegisteredWithSubsystem : 1 = false;
 
 	/** Used to detect incorrect registration/unregistration behavior */
-	UPROPERTY(Transient, DuplicateTransient, NonTransactional)
+	UPROPERTY(Transient, DuplicateTransient, NonTransactional, SkipSerialization)
 	TObjectPtr<ULandscapeSubsystem> RegisteredToSubsystem;
 
 	/** Data set at creation time */
