Clean up NetGUID error messages that appear when using UProxyNetDriver.

diff --git a/Engine/Plugins/Runtime/MultiServerReplication/Source/MultiServerReplication/Private/MultiServerProxy.cpp b/Engine/Plugins/Runtime/MultiServerReplication/Source/MultiServerReplication/Private/MultiServerProxy.cpp
--- a/Engine/Plugins/Runtime/MultiServerReplication/Source/MultiServerReplication/Private/MultiServerProxy.cpp
+++ b/Engine/Plugins/Runtime/MultiServerReplication/Source/MultiServerReplication/Private/MultiServerProxy.cpp
@@ -32,6 +32,23 @@ FProxyNetGUIDCache::FProxyNetGUIDCache(UNetDriver* NetDriver) : FNetGUIDCache(Ne
 }
 
 FNetworkGUID FProxyNetGUIDCache::AssignNewNetGUID_Server(UObject* Object)
+{
+	// The proxy's NetGUID cache does not assign NetGUIDs directly but rather looks up the corresponding
+	// NetGUID for the object in the shared backend NetGUID cache - this cache is populated using state
+	// replicated from the backend game servers which, in turn, have the authority to assign NetGUIDs.
+	//
+	// What this means is that this function can potentially return an invalid NetGUID or default NetGUID in 
+	// certain cases where the object is not in the shared backend NetGUID cache.
+	return LookupNetGUIDFromBackendCache(Object);
+}
+
+FNetworkGUID FProxyNetGUIDCache::AssignNewNetGUIDFromPath_Server(const FString& PathName, UObject* ObjOuter, UClass* ObjClass)
+{
+	ensureMsgf(0, TEXT("UProxyNetDriver does not support assinging NetGUIDs from a path."));
+	return FNetworkGUID::GetDefault();
+}
+
+FNetworkGUID FProxyNetGUIDCache::LookupNetGUIDFromBackendCache(UObject *Object)
 {
 	UProxyNetDriver* ProxyNetDriver = Cast<UProxyNetDriver>(Driver);
 	
@@ -52,19 +69,21 @@ FNetworkGUID FProxyNetGUIDCache::AssignNewNetGUID_Server(UObject* Object)
 			}
 			else
 			{
-				UE_LOG(LogNetProxy, Error, TEXT("NetGUID: Failed to find a NetGUID for object %s in the shared NetGUID cache."), *Object->GetName());
-				return FNetworkGUID::GetDefault();
+				// This function is using UObject::IsNameStableForNetworking() instead of FNetGUIDCache::IsDynamicObject(), because IsDynamicObject() considers if the outer object 
+				// has an absolute path which might not be the case if Object is a static object component (e.g. render mesh) for a dynamic object (e.g. pawns and characters).
+				if (ensureMsgf(Object->IsNameStableForNetworking(), TEXT("NetGUID: Failed to find a NetGUID for dynamic object %s in the shared NetGUID cache."), *Object->GetName()))
+				{
+					// This can occur when the proxy receives an RPC from the client with an RPC argument that is a static object reference. In this case the proxy 
+					// hasn't heard about this object from the game server yet so it wouldn't know the NetGUID. Since static object references use a path rather than
+					// a NetGUID the object will still be successfully found by the package map and forwarded to the game server.
+					UE_LOG(LogNetProxy, Log, TEXT("NetGUID: Static object %s in not in the shared NetGUID cache yet."), *Object->GetName());
+					return FNetworkGUID::GetDefault();
+				}
 			}
 		}
 	}
 
-	return FNetworkGUID::GetDefault();
-}
-
-FNetworkGUID FProxyNetGUIDCache::AssignNewNetGUIDFromPath_Server(const FString& PathName, UObject* ObjOuter, UClass* ObjClass)
-{
-	ensureMsgf(0, TEXT("UProxyNetDriver does not support assinging NetGUIDs from a path."));
-	return FNetworkGUID::GetDefault();
+	return FNetworkGUID();
 }
 
 FProxyBackendNetGUIDCache::FProxyBackendNetGUIDCache(UNetDriver* NetDriver) : FNetGUIDCache(NetDriver)
