Forgotten header file in CL 42339761.

diff --git a/Engine/Plugins/Runtime/MultiServerReplication/Source/MultiServerReplication/Public/MultiServerProxy.h b/Engine/Plugins/Runtime/MultiServerReplication/Source/MultiServerReplication/Public/MultiServerProxy.h
--- a/Engine/Plugins/Runtime/MultiServerReplication/Source/MultiServerReplication/Public/MultiServerProxy.h
+++ b/Engine/Plugins/Runtime/MultiServerReplication/Source/MultiServerReplication/Public/MultiServerProxy.h
@@ -30,6 +30,9 @@ public:
 
 	virtual FNetworkGUID AssignNewNetGUID_Server(UObject* Object);
 	virtual FNetworkGUID AssignNewNetGUIDFromPath_Server(const FString& PathName, UObject* ObjOuter, UClass* ObjClass);
+
+private:
+	FNetworkGUID LookupNetGUIDFromBackendCache(UObject *Object);
 };
 
 class FProxyBackendNetGUIDCache : public FNetGUIDCache
