Remove old code from MultiServerProxy and clean up the logic to make the connection logic more clear.

diff --git a/Engine/Plugins/Runtime/MultiServerReplication/Source/MultiServerReplication/Private/MultiServerProxy.cpp b/Engine/Plugins/Runtime/MultiServerReplication/Source/MultiServerReplication/Private/MultiServerProxy.cpp
--- a/Engine/Plugins/Runtime/MultiServerReplication/Source/MultiServerReplication/Private/MultiServerProxy.cpp
+++ b/Engine/Plugins/Runtime/MultiServerReplication/Source/MultiServerReplication/Private/MultiServerProxy.cpp
@@ -11,18 +11,30 @@
 
 DEFINE_LOG_CATEGORY_STATIC(LogNetProxy, Log, All);
 
-FString LexToString(EProxyConnectionState State)
+namespace UE::MultiServerProxy::Private
 {
-	switch (State)
+	FString ConnectionToString(const UNetConnection* Connection)
 	{
-		case EProxyConnectionState::Disconnected:
-			return FString(TEXT("Disconnected"));
-		case EProxyConnectionState::ConnectingPrimary:
-			return FString(TEXT("ConnectingPrimary"));
-		case EProxyConnectionState::ConnectedPrimary:
-			return FString(TEXT("ConnectedPrimary"));
-		default:
-			return FString(TEXT("Unknown"));
+		if (Connection)
+		{
+			return FString::Printf(TEXT("%s:%s"), *GetNameSafe(Connection->GetDriver()), *GetNameSafe(Connection));
+		}
+		else
+		{
+			return FString::Printf(TEXT("None:None"));
+		}
+	}
+
+	FString PlayerControllerToString(const APlayerController* PlayerController)
+	{
+		if (PlayerController)
+		{
+			return FString::Printf(TEXT("%s:%s"), *GetNameSafe(PlayerController), *GetNameSafe(PlayerController->GetClass()));
+		}
+		else
+		{
+			return FString::Printf(TEXT("None:None"));
+		}
 	}
 }
 
@@ -197,8 +209,6 @@ void UProxyListenerNotify::ConnectToGameServer(UNetConnection* ProxyConnection,
 											   FGameServerConnectionState* GameServerConnectionState,
 											   UE::Net::EJoinFlags Flags)
 {
-	ProxyNetDriver->ProxyConnectionState.Add(ProxyConnection, EProxyConnectionState::Disconnected);
-
 	// If this is the first connection to the game server, instantiate the backend network driver that will manage all
 	// connections from proxy connections to that server.
 	if (GameServerConnectionState->NetDriver == nullptr)
@@ -226,8 +236,6 @@ void UProxyListenerNotify::ConnectToGameServer(UNetConnection* ProxyConnection,
 	ULocalPlayer* NewPlayer = NewObject<ULocalPlayer>(GEngine, GEngine->LocalPlayerClass);
 	ProxyNetDriver->GetWorld()->GetGameInstance()->AddLocalPlayer(NewPlayer, GameServerClientId);
 
-	ProxyNetDriver->ProxyPlayerToProxyConnection.Add(NewPlayer, ProxyConnection);
-
 	// The new player will use the same unique identifier as the incoming proxy connection so that it will be propogated up to 
 	// the game servers through UNetConnection::PlayerId. This way each game server's incoming connection will have a PlayerId 
 	// that corresponds to a client connected to the proxy.
@@ -311,18 +319,23 @@ void UProxyListenerNotify::ConnectToGameServer(UNetConnection* ProxyConnection,
 	// the child connection hasn't been created yet. Once the handshake is complete, it's expected that this entry will
 	// be updated with the new child connection.
 	FMultiServerProxyInternalConnectionRoute Route;
+	Route.ClientHandshakeId = ClientHandshakeId;
 	Route.ProxyConnection = ProxyConnection;
-	Route.Player = NewPlayer;
 	Route.ParentGameServerConnection = GameServerConnection;
-
-	ProxyNetDriver->ClientHandshakeIdToInternalRoute.Add(ClientHandshakeId, Route);
+	Route.GameServerConnection = nullptr;
+	Route.Player = NewPlayer;
+	Route.PlayerController = nullptr;
 
 	if (!EnumHasAllFlags(Flags, UE::Net::EJoinFlags::NoPawn))
 	{
-		ProxyNetDriver->RegisterPrimaryGameServer(ProxyConnection, NewPlayer, GameServerConnectionStateIndex);
+		Route.State = EProxyConnectionState::ConnectingPrimary;
 	}
-
-	ProxyNetDriver->ProxyConnectionState[ProxyConnection] = EProxyConnectionState::ConnectingPrimary;
+	else
+	{
+		Route.State = EProxyConnectionState::Connecting;
+	}
+	
+	ProxyNetDriver->InternalRoutes.Add(Route);
 }
 
 void UProxyBackendNetConnection::HandleClientPlayer(APlayerController* NewPlayerController, UNetConnection* GameServerConnection)
@@ -394,23 +407,20 @@ void UProxyBackendNetDriver::ForwardRemoteFunction(UObject* RootObject, UObject*
 		return;
 	}
 
-	// If the owner is a PlayerController it is one that represents a connection to a game server and must be mapped
-	// to a PlayerController for an incoming proxy connection. If no such mapping exists it means that the RPC comes 
-	// from a game server that isn't the primary game server.
-	if (APlayerController* GameServerController = Cast<APlayerController>(OwningActor))
+	// Drop RPCs that don't come from a player controller on a primary game server.
+	if (APlayerController* PlayerController = Cast<APlayerController>(OwningActor))
 	{
-		APlayerController* ProxyController = ProxyNetDriver->GetProxyControllerFromPrimaryGameServerController(GameServerController);
-		if (ProxyController == nullptr)
+		FMultiServerProxyInternalConnectionRoute* Route = ProxyNetDriver->GetRouteWithPlayerController(PlayerController);
+		if (!ensure(Route))
 		{
-			UE_LOG(LogNetProxy, Verbose, TEXT("RPC: Ignoring %s on %s from game server %s because it's owning game server controller %s cannot map to a proxy controller."), *Function->GetName(), *RootObject->GetName(), *GetName(), *GameServerController->GetName());
 			return;
 		}
-		else
+
+		if (Route->State != EProxyConnectionState::ConnectedPrimary)
 		{
-			UE_LOG(LogNetProxy, VeryVerbose, TEXT("RPC: Remapping game server controller %s to proxy controller %s when forwarding RPC %s on %s."), *GameServerController->GetName(), *ProxyController->GetName(), *Function->GetName(), *RootObject->GetName());
+			UE_LOG(LogNetProxy, Verbose, TEXT("RPC: Ignoring %s on %s from game server %s because it's owning player controller %s cannot map to a primary route."), *Function->GetName(), *RootObject->GetName(), *GetName(), *PlayerController->GetName());
+			return;
 		}
-
-		OwningActor = ProxyController;
 	}
 
 	ULocalPlayer* OwningPlayer = Cast<ULocalPlayer>(OwningActor->GetNetOwningPlayerAnyRole());
@@ -648,7 +658,7 @@ void UProxyNetDriver::ForwardRemoteFunction(UObject* RootObject, UObject* SubObj
 
 	UE_LOG(LogNetProxy, VeryVerbose, TEXT("RPC: Pre-forwarding %s on %s from proxy connection to game server."), *Function->GetName(), *OwningActor->GetName());
 
-	if (TObjectPtr<APlayerController>* GamePlayerControllerPtr = ProxyPlayerToPrimaryGameServerPlayerController.Find(OwningPlayer))
+	if (TObjectPtr<const APlayerController>* GamePlayerControllerPtr = PlayerToPrimaryGameServerPlayerController.Find(OwningPlayer))
 	{
 		UNetConnection* GameConnection = GamePlayerControllerPtr->Get()->NetConnection;
 		if (GameConnection)
@@ -761,38 +771,7 @@ void UProxyNetDriver::InternalProcessRemoteFunction(
 
 int32 UProxyNetDriver::ServerReplicateActors(float DeltaSeconds)
 {
-	// The owners are actors in the world are going to be game server player controlers. Since these controllers 
-	// aren't replicated ot the client, and the client has it's own proxy player controller, we need to map an
-	// acor's game server player controller owner to the corresponding proxy player controller during replication.
-	TMap<AActor*, AActor*> OriginalActorOwners;
-	for (const TSharedPtr<FNetworkObjectInfo>& ObjectInfo : GetNetworkObjectList().GetAllObjects())
-	{
-		AActor* Owner = ObjectInfo->Actor->Owner;
-		if (APlayerController* GameServerController = Cast<APlayerController>(Owner))
-		{
-			OriginalActorOwners.Add(ObjectInfo->Actor, GameServerController);
-
-			APlayerController* ProxyController = GetProxyControllerFromPrimaryGameServerController(GameServerController);
-			AActor* Actor = const_cast<TSharedPtr<FNetworkObjectInfo>*>(&ObjectInfo)->Get()->Actor;
-			Actor->Owner = ProxyController;
-		}
-	}
-
-	int32 ActorCount = Super::ServerReplicateActors(DeltaSeconds);
-
-	// Restore actor owners to game server player controllers.
-	for (TPair<AActor*, AActor*>& Pair : OriginalActorOwners)
-	{
-		AActor* OriginalActor = Pair.Key;
-		AActor* OriginalOwner = Pair.Value;
-		if (const TSharedPtr<FNetworkObjectInfo>* ObjectInfo = GetNetworkObjectList().GetAllObjects().Find(OriginalActor))
-		{
-			AActor* Actor = const_cast<TSharedPtr<FNetworkObjectInfo>*>(ObjectInfo)->Get()->Actor;
-			Actor->Owner = OriginalOwner;
-		}
-	}
-
-	return ActorCount;
+	return Super::ServerReplicateActors(DeltaSeconds);
 }
 
 bool UProxyNetDriver::CanDowngradeActorRole(UNetConnection* ProxyConnection, AActor* Actor) const
@@ -856,207 +835,201 @@ void UProxyNetDriver::DisableActorLogicAndGameCode()
 
 void UProxyNetDriver::GameServerAssignPlayerController(UNetConnection* ChildGameServerConnection, UNetConnection* NotUsedParentGameServerConnection, APlayerController* GameServerPlayerController)
 {
-	UE_LOG(LogNetProxy, Log, TEXT("Received a new player controller %s:%s for connection %s:%s (client_handshake_id=%u)."), 
+	UE_LOG(LogNetProxy, Log, TEXT("[%d] Received a new player controller %s:%s for connection %s:%s."), 
+		   GameServerPlayerController->GetClientHandshakeId(),
 		   *GameServerPlayerController->GetName(), *GameServerPlayerController->GetClass()->GetName(),
-		   *ChildGameServerConnection->GetDriver()->GetName(), *ChildGameServerConnection->GetName(),
-		   GameServerPlayerController->GetClientHandshakeId());
-	
-	// The game server player controller is used as a client connection to the game server so must be ROLE_AutonomousProxy.
-	GameServerPlayerController->SetRole(ROLE_AutonomousProxy);
+		   *ChildGameServerConnection->GetDriver()->GetName(), *ChildGameServerConnection->GetName());
 
 	const uint32 ClientHandshakeId = GameServerPlayerController->GetClientHandshakeId();
-	FMultiServerProxyInternalConnectionRoute* RoutePtr = ClientHandshakeIdToInternalRoute.Find(ClientHandshakeId);
-	if (RoutePtr == nullptr)
+	FMultiServerProxyInternalConnectionRoute* Route = GetRouteWithClientHandshakeId(ClientHandshakeId);
+	if (!ensure(Route))
 	{
-		UE_LOG(LogNetProxy, Error, TEXT("Failed to find a mapping for game server connection request (client_handshake_id=%u)"), ClientHandshakeId);
 		return;
 	}
 
-	// Cache this value locally incase ClientHandshakeIdToInternalRoute changes and invalidates the pointer.
-	FMultiServerProxyInternalConnectionRoute Route = *RoutePtr;
-	
-	UE_LOG(LogNetProxy, Log, TEXT("Found internal route (client_handshake_id=%u player=%s proxy_connection=%s:%s parent_game_server_connection=%s:%s)"), 
-		   ClientHandshakeId,
-		   *GetNameSafe(Route.Player),
-			Route.ProxyConnection ? *GetNameSafe(Route.ProxyConnection->GetDriver()) : TEXT("None"),
-		   *GetNameSafe(Route.ProxyConnection),
-			Route.ParentGameServerConnection ? *GetNameSafe(Route.ParentGameServerConnection->GetDriver()) : TEXT("None"),
-		   *GetNameSafe(Route.ParentGameServerConnection));
-
-	ULocalPlayer* Player = Route.Player;
-		
-	// The player should only reference a proxy player controller if this game server connection
-	// is the primary game server for the proxy connection (decided later in this function).
-	Player->PlayerController = nullptr;
-	
-	// Since an existing player controller can be re-assigned it must be detached from any existing connections and players.
-	if (GameServerPlayerController->NetConnection != nullptr)
+	EProxyConnectionState State = Route->State;
+
+	switch (State)
 	{
-		UE_LOG(LogNetProxy, Log, TEXT("Player controller %s already assigned to connection %s:%s."),
-			   *GameServerPlayerController->GetName(),
-			   *GetNameSafe(GameServerPlayerController->NetConnection->Driver),
-			   *GameServerPlayerController->NetConnection->GetName());
+		case EProxyConnectionState::Disconnected:
+			break;
+		case EProxyConnectionState::Connecting:
+			ProcessConnectionState_Connecting(Route, ChildGameServerConnection, GameServerPlayerController); 
+			break;
+		case EProxyConnectionState::ConnectingPrimary:
+			ProcessConnectionState_ConnectingPrimary(Route, ChildGameServerConnection, GameServerPlayerController);
+			break;
+		case EProxyConnectionState::Connected:
+			ProcessConnectionState_Connected(Route, ChildGameServerConnection, GameServerPlayerController);
+			break;
+		case EProxyConnectionState::ConnectedPrimary:
+			ProcessConnectionState_ConnectedPrimary(Route, ChildGameServerConnection, GameServerPlayerController);
+			break;
+		case EProxyConnectionState::PendingReassign:
+			ProcessConnectionState_Reassign(Route, ChildGameServerConnection, GameServerPlayerController);
+			break;
+	}
 
-		DetachPlayerController(GameServerPlayerController);
+	// If this was a parent connection send any join requests for any pending multiplexed connections.
+	if (ChildGameServerConnection == Route->ParentGameServerConnection)
+	{
+		for (FGameServerConnectionState& GameServerConnectionState : GameServerConnections)
+		{
+			if (GameServerConnectionState.NetDriver->ServerConnection == Route->ParentGameServerConnection)
+			{
+				FlushSplitJoinRequests(&GameServerConnectionState);
+				break;
+			}
+		}
 	}
+}
+
+TSharedPtr<FNetGUIDCache> UProxyNetDriver::GetSharedBackendNetGuidCache()
+{
+	return SharedBackendNetGuidCache;
+}
 
-	// If a new player controller is spawned on the proxy from a game server it's assumed it will replace any
-	// existing game server player controllers used by connections and players on the proxy. For this reason
-	// always detach an existing game server player controllers from the connection.
-	DetachPlayerController(ChildGameServerConnection->PlayerController);
-	AttachPlayerController(GameServerPlayerController, ChildGameServerConnection, Player);
+void UProxyNetDriver::FinalizeGameServerConnection(FMultiServerProxyInternalConnectionRoute* Route, UNetConnection* GameServerConnection, APlayerController* PlayerController)
+{
+	ensure(Route->PlayerController == nullptr);
+	ensure(Route->GameServerConnection == nullptr);
+	ensure(Route->ClientHandshakeId == PlayerController->GetClientHandshakeId());
+	ensure(Route->ProxyConnection);
+
+	PlayerController->Player = Route->Player;
+	PlayerController->NetConnection = GameServerConnection;
+	PlayerController->SetRole(ROLE_AutonomousProxy);
+
+	GameServerConnection->SetConnectionState(EConnectionState::USOCK_Open);
+	GameServerConnection->SetClientHandshakeId(Route->ClientHandshakeId);
+	GameServerConnection->PlayerController = PlayerController;
+	GameServerConnection->OwningActor = PlayerController;
+	GameServerConnection->LastReceiveTime = GetElapsedTime();
+
+	Route->GameServerConnection = GameServerConnection;
+	Route->PlayerController = PlayerController;
+}
 
-	// Find the proxy connection that established this connection to the game server.
-	TObjectPtr<UNetConnection> ProxyConnection = Route.ProxyConnection;
+void UProxyNetDriver::ReassignPlayerController(const FMultiServerProxyInternalConnectionRoute* Route, UNetConnection* GameServerConnection, APlayerController* PlayerController)
+{
+	ensure(Route->State == EProxyConnectionState::Connected || Route->State == EProxyConnectionState::ConnectedPrimary || Route->State == EProxyConnectionState::PendingReassign);
+	ensure(Route->ClientHandshakeId == PlayerController->GetClientHandshakeId());
+	ensure(Route->GameServerConnection->GetClientHandshakeId() == PlayerController->GetClientHandshakeId());
+	ensure(Route->GameServerConnection == GameServerConnection);
 
-	if (ensure(ProxyConnectionState.Contains(ProxyConnection)))
+	FMultiServerProxyInternalConnectionRoute* PrevRoute = GetRouteWithPlayerController(PlayerController);
+	if (PrevRoute)
 	{
-		EProxyConnectionState State = ProxyConnectionState[ProxyConnection];
+		UE_LOG(LogNetProxy, Log, TEXT("[%u] Clearing the player controller %s from the previous route with client handshake id %u"), 
+			   Route->ClientHandshakeId,
+			   *UE::MultiServerProxy::Private::PlayerControllerToString(PlayerController),
+			   PrevRoute->ClientHandshakeId);
 
-		UE_LOG(LogNetProxy, Log, TEXT("Attempting to assign game server player controller %s to proxy connection %s:%s (state=%s)"),
-			*GetNameSafe(GameServerPlayerController),
-			ProxyConnection ? *GetNameSafe(ProxyConnection->GetDriver()) : TEXT("None"),
-			*GetNameSafe(ProxyConnection),
-			*LexToString(State));
+		
+		// When a player controller is assigned to a new route it's assumed that the route that is currently using that player controller will
+		// also have a new controller assigned. The PendingReassign state marks that other route as in a pending state that is waiting for it's
+		// new controller.
+		PrevRoute->PlayerController = nullptr;
+		PrevRoute->State = EProxyConnectionState::PendingReassign;
+	}
+	else
+	{
+		UE_LOG(LogNetProxy, Log, TEXT("[%u] Player controller %s was not bound to another route."), 
+			   Route->ClientHandshakeId,
+			   *UE::MultiServerProxy::Private::PlayerControllerToString(PlayerController));
+	}
 
-		switch (State)
-		{
-			case EProxyConnectionState::Disconnected:
-				break;
+	const APlayerController* PrevPlayerController = Route->PlayerController;
 
-			case EProxyConnectionState::ConnectingPrimary:
-				{
-					if (int32* PlayerPrimaryGameServerIdxPtr = ProxyPlayerToPrimaryGameServer.Find(Player))
-					{
-						int32 PlayerPrimaryGameServerIdx = *PlayerPrimaryGameServerIdxPtr;
-
-						ensure(ProxyConnection->PlayerController == nullptr);
-
-						if (ensure(PlayerPrimaryGameServerIdx < GameServerConnections.Num()))
-						{
-							FGameServerConnectionState& PlayerPrimaryGameServer = GameServerConnections[PlayerPrimaryGameServerIdx];
-							if (PlayerPrimaryGameServer.NetDriver->ServerConnection == Route.ParentGameServerConnection)
-							{
-								UPlayer* PreviousPlayer = GameServerPlayerController->Player;
-			
-								APlayerController* ProxyPlayerController = GameServerPlayerController;
-								{
-									DetachPlayerController(ProxyConnection->PlayerController);
-
-									ProxyConnection->SetConnectionState(EConnectionState::USOCK_Open);
-									ProxyConnection->SetClientHandshakeId(ProxyPlayerController->GetClientHandshakeId());
-									ProxyConnection->PlayerController = ProxyPlayerController;
-									ProxyConnection->OwningActor = ProxyPlayerController;
-
-									// The player always points to the proxy player controller if that player is
-									// considered the primary player for the proxy connection.
-									Player->PlayerController = ProxyPlayerController;
-
-									ProxyPlayerToPrimaryGameServerPlayerController.Add(Player, GameServerPlayerController);
-
-									UE_LOG(LogNetProxy, Log, TEXT("Successfully connected proxy connection %s:%s to primary game server connection %s:%s (game_server_controller=%s:%s proxy_controller=%s:%s player=%s client_handshake_id=%u)"),
-										*ProxyConnection->GetDriver()->GetName(),
-										*ProxyConnection->GetName(),
-										*ChildGameServerConnection->GetDriver()->GetName(),
-										*ChildGameServerConnection->GetName(),
-										*GameServerPlayerController->GetName(),
-										*GameServerPlayerController->GetClass()->GetName(),
-										*ProxyPlayerController->GetName(),
-										*ProxyPlayerController->GetClass()->GetName(),
-										*Player->GetName(),
-										ClientHandshakeId);
-
-									ProxyConnectionState[ProxyConnection] = EProxyConnectionState::ConnectedPrimary;
-								}
-							}
-						}
-					}
-					else
-					{
-						UE_LOG(LogNetProxy, Log, TEXT("Game server player controller %s not from proxy connection %s:%s primary game server"),
-							*GameServerPlayerController->GetName(),
-							*ProxyConnection->GetDriver()->GetName(),
-							*ProxyConnection->GetName());
-					}
-				}
-				break;
+	PlayerController->NetConnection = GameServerConnection;
+	PlayerController->Player = Route->Player;
+	Route->PlayerController = PlayerController;
 
-			case EProxyConnectionState::ConnectedPrimary:
-				{
-					if (!GameServerPlayerController->IsA(ANoPawnPlayerController::StaticClass()))
-					{
-						TObjectPtr<ULocalPlayer>* ExistingPrimaryPlayerPtr = ProxyConnectionToPrimaryPlayer.Find(ProxyConnection);
-						if (ExistingPrimaryPlayerPtr)
-						{
-							ULocalPlayer* ExistingPrimaryPlayer = *ExistingPrimaryPlayerPtr;
-
-							if (ProxyPlayerToPrimaryGameServerPlayerController.Contains(ExistingPrimaryPlayer))
-							{
-								int32* PreviousPrimaryGameServerIdxPtr = ProxyPlayerToPrimaryGameServer.Find(ExistingPrimaryPlayer);
-								if (ensure(PreviousPrimaryGameServerIdxPtr))
-								{
-									int32 PreviousPrimaryGameServerIdx = *PreviousPrimaryGameServerIdxPtr;
-
-									if (Player->GetCachedUniqueNetId() == ExistingPrimaryPlayer->GetCachedUniqueNetId())
-									{
-										int32 NewPrimaryGameServerIdx = GetGameServerWithPlayer(Player);
-
-										UE_LOG(LogNetProxy, Log, TEXT("Changing primary game server for proxy connection %s:%s from player %s to %s (prev_game_server_idx=%d new_game_server_idx=%d)."),
-											*ProxyConnection->Driver->GetName(),
-											*ProxyConnection->GetName(),
-											*ExistingPrimaryPlayer->GetName(),
-											*Player->GetName(),
-											PreviousPrimaryGameServerIdx,
-											NewPrimaryGameServerIdx);
-					
-										ProxyPlayerToPrimaryGameServerPlayerController.Remove(ExistingPrimaryPlayer);
-
-										DeregisterPrimaryGameServer(ProxyConnection, ExistingPrimaryPlayer);
-										RegisterPrimaryGameServer(ProxyConnection, Player, NewPrimaryGameServerIdx);
-				
-										ExistingPrimaryPlayer->PlayerController = nullptr;
-										Player->PlayerController = ProxyConnection->PlayerController;
-										ProxyConnection->PlayerController->Player = Player;
-
-										ProxyPlayerToPrimaryGameServerPlayerController.Add(Player, GameServerPlayerController);
-									}
-								}
-							}
-						}
-					}
-				}
-				break;
+	if (PlayerController->IsA(ANoPawnPlayerController::StaticClass()))
+	{
+		UE_LOG(LogNetProxy, Log, TEXT("[%u] Deregistering player %s as a primary game server."), 
+			   Route->ClientHandshakeId,
+			   *Route->Player->GetName());
 
-		}
+		PlayerToPrimaryGameServerPlayerController.Remove(Route->Player);
+		Route->State = EProxyConnectionState::Connected;
 	}
-			
-	// Update the internal connection mapping to use the new child game server connection (if it's different from the parent connection).
-	ClientHandshakeIdToInternalRoute.Remove(ClientHandshakeId);
-	ClientHandshakeIdToInternalRoute.Add(ClientHandshakeId, Route);
-
-	// If this was a parent connection send any join requests for any pending multiplexed connections.
-	if (ChildGameServerConnection == Route.ParentGameServerConnection)
+	else
 	{
-		FGameServerConnectionState* CurrGameServerConnectionState = nullptr;
-		for (FGameServerConnectionState& GameServerConnectionState : GameServerConnections)
-		{
-			if (GameServerConnectionState.NetDriver->ServerConnection == Route.ParentGameServerConnection)
-			{
-				CurrGameServerConnectionState = &GameServerConnectionState;
-				break;
-			}
-		}
+		UE_LOG(LogNetProxy, Log, TEXT("[%u] Registering player %s as a primary game server using player controller %s."), 
+			   Route->ClientHandshakeId,
+			   *Route->Player->GetName(),
+			   *UE::MultiServerProxy::Private::PlayerControllerToString(Route->PlayerController));
 
-		if (ensure(CurrGameServerConnectionState))
-		{
-			FlushSplitJoinRequests(CurrGameServerConnectionState);
-		}
+		PlayerToPrimaryGameServerPlayerController.Add(Route->Player, Route->PlayerController);
+		Route->State = EProxyConnectionState::ConnectedPrimary;
 	}
+
+	UE_LOG(LogNetProxy, Log, TEXT("[%u] Replacing player controller %s with %s on route with game server connection %s and player %s"),
+		   Route->ClientHandshakeId,
+		   *UE::MultiServerProxy::Private::PlayerControllerToString(PrevPlayerController),
+		   *UE::MultiServerProxy::Private::PlayerControllerToString(Route->PlayerController),
+		   *UE::MultiServerProxy::Private::ConnectionToString(Route->GameServerConnection),
+		   *GetNameSafe(Route->Player));
 }
 
-TSharedPtr<FNetGUIDCache> UProxyNetDriver::GetSharedBackendNetGuidCache()
+void UProxyNetDriver::ProcessConnectionState_Connecting(FMultiServerProxyInternalConnectionRoute* Route, UNetConnection* GameServerConnection, APlayerController* PlayerController)
 {
-	return SharedBackendNetGuidCache;
+	ensure(Route->State == EProxyConnectionState::Connecting);
+
+	FinalizeGameServerConnection(Route, GameServerConnection, PlayerController);
+
+	Route->State = EProxyConnectionState::Connected;
+
+	UE_LOG(LogNetProxy, Log, TEXT("[%d] Successfully connected proxy connection %s to non-primary game server connection %s using player controller %s and player %s."),
+		   Route->ClientHandshakeId,
+		   *UE::MultiServerProxy::Private::ConnectionToString(Route->ProxyConnection),
+		   *UE::MultiServerProxy::Private::ConnectionToString(GameServerConnection),
+		   *UE::MultiServerProxy::Private::PlayerControllerToString(PlayerController),
+		   *Route->Player->GetName());
+}
+
+void UProxyNetDriver::ProcessConnectionState_ConnectingPrimary(FMultiServerProxyInternalConnectionRoute* Route, UNetConnection* GameServerConnection, APlayerController* PlayerController)
+{
+	ensure(Route->State == EProxyConnectionState::ConnectingPrimary);
+
+	FinalizeGameServerConnection(Route, GameServerConnection, PlayerController);
+
+	Route->ProxyConnection->SetConnectionState(EConnectionState::USOCK_Open);
+	Route->ProxyConnection->PlayerController = PlayerController;
+	Route->ProxyConnection->OwningActor = PlayerController;
+	Route->ProxyConnection->LastReceiveTime = GetElapsedTime();
+	
+	UE_LOG(LogNetProxy, Log, TEXT("[%u] Registering player %s as a primary game server using player controller %s."), 
+		   Route->ClientHandshakeId,
+		   *Route->Player->GetName(),
+		   *UE::MultiServerProxy::Private::PlayerControllerToString(Route->PlayerController));
+
+	PlayerToPrimaryGameServerPlayerController.Add(Route->Player, Route->PlayerController);
+
+	Route->State = EProxyConnectionState::ConnectedPrimary;
+
+	UE_LOG(LogNetProxy, Log, TEXT("[%d] Successfully connected proxy connection %s to primary game server connection %s using player controller %s and player %s."),
+		   Route->ClientHandshakeId,
+		   *UE::MultiServerProxy::Private::ConnectionToString(Route->ProxyConnection),
+		   *UE::MultiServerProxy::Private::ConnectionToString(GameServerConnection),
+		   *UE::MultiServerProxy::Private::PlayerControllerToString(PlayerController),
+		   *Route->Player->GetName());
+}
+
+void UProxyNetDriver::ProcessConnectionState_Connected(const FMultiServerProxyInternalConnectionRoute* Route, UNetConnection* GameServerConnection, APlayerController* PlayerController)
+{
+	ReassignPlayerController(Route, GameServerConnection, PlayerController);
+}
+
+void UProxyNetDriver::ProcessConnectionState_ConnectedPrimary(const FMultiServerProxyInternalConnectionRoute* Route, UNetConnection* GameServerConnection, APlayerController* PlayerController)
+{
+	ReassignPlayerController(Route, GameServerConnection, PlayerController);
+}
+
+void UProxyNetDriver::ProcessConnectionState_Reassign(const FMultiServerProxyInternalConnectionRoute* Route, UNetConnection* GameServerConnection, APlayerController* PlayerController)
+{
+	ReassignPlayerController(Route, GameServerConnection, PlayerController);
 }
 
 void UProxyNetDriver::FlushSplitJoinRequests(FGameServerConnectionState* GameServerConnectionState)
@@ -1089,142 +1062,54 @@ void UProxyNetDriver::FlushSplitJoinRequests(FGameServerConnectionState* GameSer
 	GameServerConnectionState->PendingSplitJoinRequests.Reset();
 }
 
-int32 UProxyNetDriver::GetGameServerWithPlayer(ULocalPlayer* Player)
+UNetConnection* UProxyNetDriver::GetProxyConnectionFromPrimaryPlayer(ULocalPlayer* Player)
 {
-	for (int32 GameServerIdx = 0; GameServerIdx < GameServerConnections.Num(); GameServerIdx++)
+	TObjectPtr<const APlayerController>* GameServerPlayerControllerPtr = PlayerToPrimaryGameServerPlayerController.Find(Player);
+	if (!GameServerPlayerControllerPtr)
 	{
-		if (GameServerConnections[GameServerIdx].Players.Find(Player) != INDEX_NONE)
-		{
-			return GameServerIdx;
-		}
+		return nullptr;
 	}
 
-	return -1;
-}
-
-void UProxyNetDriver::DeregisterPrimaryGameServer(UNetConnection* ProxyConnection, ULocalPlayer* PreviousPrimaryPlayer)
-{
-	if (ensure(ProxyConnection->PlayerController))
-	{
-		FUniqueNetIdRepl& PlayerId = ProxyConnection->PlayerId;
-
-		UE_LOG(LogNetProxy, Log, TEXT("Clearing player %s primary game server (proxy_connection=%s:%s)"),
-			*PreviousPrimaryPlayer->GetName(),
-			*GetNameSafe(ProxyConnection->GetDriver()),
-			*ProxyConnection->GetName());
-
-		ensure(ProxyPlayerToPrimaryGameServer.Contains(PreviousPrimaryPlayer));
-		
-		ProxyPlayerToPrimaryGameServer.Remove(PreviousPrimaryPlayer);
-		ProxyConnectionToPrimaryPlayer.Remove(ProxyConnection);
-	}
-}
+	const APlayerController* GameServerPlayerController = GameServerPlayerControllerPtr->Get();
 
-void UProxyNetDriver::RegisterPrimaryGameServer(UNetConnection* ProxyConnection, class ULocalPlayer* PrimaryPlayer, int32 GameServerConnectionStateIndex)
-{
-	if (ensure(GameServerConnectionStateIndex < GameServerConnections.Num()))
+	FMultiServerProxyInternalConnectionRoute* Route = GetRouteWithPlayerController(GameServerPlayerController);
+	if (!Route)
 	{
-		FUniqueNetIdRepl& PlayerId = ProxyConnection->PlayerId;
-
-		UE_LOG(LogNetProxy, Log, TEXT("Configuring primary game server %d for player (player=%s proxy_connection=%s:%s)"),
-			GameServerConnectionStateIndex,
-			*PrimaryPlayer->GetName(),
-			*ProxyConnection->GetDriver()->GetName(),
-			*ProxyConnection->GetName());
-
-		ensure(!ProxyPlayerToPrimaryGameServer.Contains(PrimaryPlayer));
-
-		ProxyPlayerToPrimaryGameServer.Add(PrimaryPlayer, GameServerConnectionStateIndex);
-		ProxyConnectionToPrimaryPlayer.Add(ProxyConnection, PrimaryPlayer);
+		return nullptr;
 	}
-}
-
-void UProxyNetDriver::AttachPlayerController(APlayerController* PlayerController, UNetConnection* Connection, ULocalPlayer* Player)
-{
-	// The presumed relationships between connection, player controller and player:
-	//	
-	//	ProxyConnection <-> ProxyPlayerController -> LocalPlayer
-	//	LocalPlayer <- GameServerPlayerController <-> Game Server Connection
-	//
-	// This function doesn't associate the player with a player controller since it differs between
-	// proxy and game server player controllers.
-
-	UE_LOG(LogNetProxy, Log, TEXT("Attaching controller %s:%s on connection %s and player %s."),
-		   *PlayerController->GetName(),
-		   *PlayerController->GetClass()->GetName(),
-		   *Connection->GetName(),
-		   *Player->GetName());
 
-	Connection->SetConnectionState(EConnectionState::USOCK_Open);
-	Connection->SetClientHandshakeId(PlayerController->GetClientHandshakeId());
+	ensure(Route->Player == Player);
+	ensure(Route->State == EProxyConnectionState::ConnectedPrimary);
 
-	// It's assumed the connection is not associated with a player controller.	
-	ensure(Connection->PlayerController == nullptr);
-	ensure(Connection->OwningActor == nullptr);
-
-	Connection->PlayerController = PlayerController;
-	Connection->OwningActor = PlayerController;
-	
-	PlayerController->NetConnection = Connection;
-	PlayerController->Player = Player;
-	
-	Connection->LastReceiveTime = GetElapsedTime();
+	return Route->ProxyConnection;
 }
 
-void UProxyNetDriver::DetachPlayerController(APlayerController* PlayerController)
+FMultiServerProxyInternalConnectionRoute* UProxyNetDriver::GetRouteWithPlayerController(const APlayerController* PlayerController)
 {
-	if (PlayerController)
+	FMultiServerProxyInternalConnectionRoute* RoutePtr = nullptr;
+	for (FMultiServerProxyInternalConnectionRoute& Route : InternalRoutes)
 	{
-		UPlayer* Player = PlayerController->Player;
-
-		UNetConnection* Connection = PlayerController->NetConnection;
-		if (PlayerController->NetConnection)
-		{
-			Connection->PlayerController = nullptr;
-			Connection->OwningActor = nullptr;
-		}
-
-		PlayerController->NetConnection = nullptr;
-		PlayerController->Player = nullptr;
-
-		// Only detach the player controller from the player if it's the one referenced by the player.
-		if (Player)
+		if (Route.PlayerController == PlayerController)
 		{
-			if (Player->PlayerController == PlayerController)
-			{
-				Player->PlayerController = nullptr;
-			}
+			ensure(RoutePtr == nullptr);
+			RoutePtr = &Route;
 		}
-
-		UE_LOG(LogNetProxy, Log, TEXT("Detaching old controller %s:%s on %s attached to player %s and connection %s."),
-			   *PlayerController->GetName(),
-			   *PlayerController->GetClass()->GetName(),
-			   *GetName(),
-			   *GetNameSafe(Player),
-			   *GetNameSafe(Connection));
 	}
+
+	return RoutePtr;
 }
 
-UNetConnection* UProxyNetDriver::GetProxyConnectionFromPrimaryPlayer(UPlayer* Player)
+FMultiServerProxyInternalConnectionRoute* UProxyNetDriver::GetRouteWithClientHandshakeId(uint32 ClientHandshakeId)
 {
-	if (ULocalPlayer* LocalPlayer = Cast<ULocalPlayer>(Player))
+	FMultiServerProxyInternalConnectionRoute* RoutePtr = nullptr;
+	for (FMultiServerProxyInternalConnectionRoute& Route : InternalRoutes)
 	{
-		if (TObjectPtr<UNetConnection>* ConnectionPtr = ProxyPlayerToProxyConnection.Find(LocalPlayer))
+		if (Route.ClientHandshakeId == ClientHandshakeId)
 		{
-			return ConnectionPtr->Get();
+			ensure(RoutePtr == nullptr);
+			RoutePtr = &Route;
 		}
 	}
 
-	return nullptr;
-}
-
-APlayerController* UProxyNetDriver::GetProxyControllerFromPrimaryGameServerController(APlayerController* GameServerController)
-{
-	if (ULocalPlayer* Player = Cast<ULocalPlayer>(GameServerController->Player))
-	{
-		// The player's controller will always point to the proxy player controller.
-		return Player->PlayerController;
-	}
-
-	return nullptr;
+	return RoutePtr;
 }
\ No newline at end of file
diff --git a/Engine/Plugins/Runtime/MultiServerReplication/Source/MultiServerReplication/Public/MultiServerProxy.h b/Engine/Plugins/Runtime/MultiServerReplication/Source/MultiServerReplication/Public/MultiServerProxy.h
--- a/Engine/Plugins/Runtime/MultiServerReplication/Source/MultiServerReplication/Public/MultiServerProxy.h
+++ b/Engine/Plugins/Runtime/MultiServerReplication/Source/MultiServerReplication/Public/MultiServerProxy.h
@@ -13,12 +13,13 @@
 enum class EProxyConnectionState : uint8
 {
 	Disconnected,
+	Connecting,
 	ConnectingPrimary,
-	ConnectedPrimary
+	Connected,
+	ConnectedPrimary,
+	PendingReassign
 };
 
-FString LexToString(EProxyConnectionState State);
-
 /** 
  * A NetGUID cache used by UProxyNetDriver that doesn't assign new NetGUIDs but uses an already assigned NetGUID
  * from a backend server connection's NetGUID cache.
@@ -100,14 +101,14 @@ private:
 
 /** 
  * A route from a proxy client connection to a single game server.
- *
- * This route binds a physical proxy client connection to a logical (multiplex) game server connection.
  */ 
 USTRUCT()
 struct FMultiServerProxyInternalConnectionRoute
 {
 	GENERATED_BODY()
 
+	uint32 ClientHandshakeId;
+
 	/** The client connection to the proxy (the beginning of the route). */
 	UPROPERTY()
 	TObjectPtr<UNetConnection> ProxyConnection;
@@ -116,9 +117,26 @@ struct FMultiServerProxyInternalConnectionRoute
 	UPROPERTY()
 	TObjectPtr<UNetConnection> ParentGameServerConnection;
 
+	/** The actual connection to the backend server. */
+	UPROPERTY()
+	TObjectPtr<UNetConnection> GameServerConnection;
+
 	/** The player used to represent this route in the game systems. */
 	UPROPERTY()
 	TObjectPtr<class ULocalPlayer> Player;
+
+	/** 
+	 * The functions below are marked as mutable to allow the rest of the structure to remain const
+	 * after the initial connection has been finalized, avoiding mistakenly changing route state 
+	 * when migrating player controllers between game servers and changing primary game server.
+	 */
+
+	/** The player controller. */
+	UPROPERTY()
+	mutable TObjectPtr<const class APlayerController> PlayerController;
+
+	/** The state of the route. */
+	mutable EProxyConnectionState State = EProxyConnectionState::Disconnected;
 };
 
 /** 
@@ -266,31 +284,7 @@ private:
 * (ROLE_SimulatedProxy or ROLE_AutonomousProxy) and will be replicated as-is to the proxy client.
 *
 * When a proxy connection (UNetConnection) is opened in UProxyNetDriver it opens a game server 
-* connection (UNetConnection) to each registered backend server. Each proxy connection and game 
-* server connection is associated with it's own instance of APlayerController. There will always 
-* be an equal number of connections and player controllers on a proxy and can be calculated with
-* this equation: i+(j*i), where i is the number of clients connected to the proxy and j is the 
-* number of registered game servers.
-*
-* The proxy and game server player controllers are related through an instance of ULocalPlayer
-* for each proxy connection to a registered game server: the equation for the number of players
-* on the proxy is (i*j), where i is the number of clients connected to the proxy and j in the 
-* number of registered game servers.
-*
-* The relationship between proxy connections and game server connections is done through pointers
-* in the respective classes:
-*
-* ProxyConnection <-> ProxyPlayerController <-> LocalPlayer <- GameServerPlayerController <-> GameServerConnection
-*
-* From this representation of the relationship we can see that ProxyConnection has a pointer to 
-* ProxyPlayerController, and ProxyPlayerController has a pointer to ProxyConnection. Using these
-* existing pointers enables the use of existing UE code and avoids the need for additional data
-* structures to hold the relationship betweeen the clients connected to the proxy
-* and registered game servers.
-*
-* IMPORTANT: There isn't a relationship from LocalPlayer to GameServerPlayerController because
-* the exiting pointer (ULocalPlayer::PlayerController) that is being used can only point to one
-* controller.
+* connection (UNetConnection) to each registered backend server.
 *
 * For each proxy connection one of the game servers are considered the primary game server. This
 * is the game server that spawns the proxy client's pawn, player controller, receive RPCs from
@@ -351,58 +345,42 @@ private:
 	// Classes in EngineTest that test the behavior of UMultiServerProxy.
 	friend class FNetTestProxyRPCRouting;
 	friend class FNetTestProxyConnectionRouting;
-	friend class FNetTestProxyMovePlayerController;
 	
 	friend class UProxyBackendNetDriver;
 
+	/** Called when the connection handshake to a game server is completed (i.e. the first player controller has been replicated on that connection). */
+	void FinalizeGameServerConnection(FMultiServerProxyInternalConnectionRoute* Route, UNetConnection* GameServerConnection, APlayerController* PlayerController);
+
+	/** Called when a player controller is switching on the remote game server (i.e. during migration). */
+	void ReassignPlayerController(const FMultiServerProxyInternalConnectionRoute* Route, UNetConnection* GameServerConnection, APlayerController* PlayerController);
+
+	void ProcessConnectionState_Connecting(FMultiServerProxyInternalConnectionRoute* Route, UNetConnection* GameServerConnection, APlayerController* PlayerController);
+	void ProcessConnectionState_ConnectingPrimary(FMultiServerProxyInternalConnectionRoute* Route, UNetConnection* GameServerConnection, APlayerController* PlayerController);
+	void ProcessConnectionState_Connected(const FMultiServerProxyInternalConnectionRoute* Route, UNetConnection* GameServerConnection, APlayerController* PlayerController);
+	void ProcessConnectionState_ConnectedPrimary(const FMultiServerProxyInternalConnectionRoute* Route, UNetConnection* GameServerConnection, APlayerController* PlayerController);
+	void ProcessConnectionState_Reassign(const FMultiServerProxyInternalConnectionRoute* Route, UNetConnection* GameServerConnection, APlayerController* PlayerController);
+
 	/** Set all of the configuration options that disable executing actor functionality or game specific code. */
 	void DisableActorLogicAndGameCode();
 	
 	/** Send any split join requests that have been queued up to the game server. */
 	void FlushSplitJoinRequests(FGameServerConnectionState* GameServerConnectionState);
 
-	/** Return a game server index that a player is connected to. */
-	int32 GetGameServerWithPlayer(class ULocalPlayer* Player);
-
-	/** Disassociate a proxy connection with a game server as a primary game server. */
-	void DeregisterPrimaryGameServer(UNetConnection* ProxyConnection, ULocalPlayer* PreviousPrimaryPlayer);
-
-	/** Associate a proxy connection with a primary game server. */
-	void RegisterPrimaryGameServer(UNetConnection* ProxyConnection, class ULocalPlayer* PrimaryPlayer, int32 GameServerConnectionStateIndex);
-
-	/** Attach a new player, player controller and connection. */
-	void AttachPlayerController(APlayerController* PlayerController, UNetConnection* Connection, ULocalPlayer* Player);
-
-	/** Detach an existing player, player controller and connection. */
-	void DetachPlayerController(APlayerController* PlayerController);
-
 	/** Return a proxy connection from a primary player (a player on a primary game server). */
-	UNetConnection* GetProxyConnectionFromPrimaryPlayer(UPlayer* Player);
-
-	/** Return a proxy player controller given a primary game server controller. */
-	APlayerController* GetProxyControllerFromPrimaryGameServerController(APlayerController* GameServerPlayerController);
+	UNetConnection* GetProxyConnectionFromPrimaryPlayer(ULocalPlayer* Player);
 	
-	/** Associate a request to connect to a game server (client handshake id)  with a route from a proxy connection to game server. */
-	UPROPERTY()
-	TMap<uint32, FMultiServerProxyInternalConnectionRoute> ClientHandshakeIdToInternalRoute;
+	/** Returns the internal route that uses a player controller. */
+	FMultiServerProxyInternalConnectionRoute* GetRouteWithPlayerController(const APlayerController* PlayerController);
 
-	/** Map a player to it's primary game server index into GameServerConnections. */
-	UPROPERTY()
-	TMap<TObjectPtr<class ULocalPlayer>, int32> ProxyPlayerToPrimaryGameServer;
+	/** Returns the internal route that uses a client handshake id. */
+	FMultiServerProxyInternalConnectionRoute* GetRouteWithClientHandshakeId(uint32 ClientHandshakeId);
 
-	UPROPERTY()
-	TMap<TObjectPtr<UNetConnection>, TObjectPtr<class ULocalPlayer>> ProxyConnectionToPrimaryPlayer;
+	/** All of the routes for connections to game servers. */
+	TArray<FMultiServerProxyInternalConnectionRoute> InternalRoutes;
 
 	/** Map a player to it's primary game server player controller. */
 	UPROPERTY()
-	TMap<TObjectPtr<class ULocalPlayer>, TObjectPtr<APlayerController>> ProxyPlayerToPrimaryGameServerPlayerController;
-
-	/** Map a player to it's proxy connection. */
-	UPROPERTY()
-	TMap<TObjectPtr<class ULocalPlayer>, TObjectPtr<UNetConnection>> ProxyPlayerToProxyConnection;
-
-	/** The state of each incoming proxy connection. */
-	TMap<TObjectPtr<UNetConnection>, EProxyConnectionState> ProxyConnectionState;
+	TMap<TObjectPtr<class ULocalPlayer>, TObjectPtr<const APlayerController>> PlayerToPrimaryGameServerPlayerController;
 
 	/** Net drivers and associated state used to connect to backend game servers. */
 	UPROPERTY()
