Support for actor replication relevancy when using the multi-server proxy and some bug fixes.

diff --git a/Engine/Plugins/Runtime/MultiServerReplication/Source/MultiServerReplication/Private/MultiServerProxy.cpp b/Engine/Plugins/Runtime/MultiServerReplication/Source/MultiServerReplication/Private/MultiServerProxy.cpp
--- a/Engine/Plugins/Runtime/MultiServerReplication/Source/MultiServerReplication/Private/MultiServerProxy.cpp
+++ b/Engine/Plugins/Runtime/MultiServerReplication/Source/MultiServerReplication/Private/MultiServerProxy.cpp
@@ -8,6 +8,7 @@
 #include "GameFramework/GameModeBase.h"
 #include "Misc/CommandLine.h"
 #include "Net/DataChannel.h"
+#include "Net/Core/Misc/NetSubObjectRegistry.h"
 
 DEFINE_LOG_CATEGORY_STATIC(LogNetProxy, Log, All);
 
@@ -730,7 +731,7 @@ void UProxyNetDriver::RemoveNetworkActor(AActor* Actor)
 	// servers but will still maintain the same netguid (derived from a remote id). For this reason,
 	// it's important to remove an object from the proxy netguid cache when it's being removed from
 	// the replication system to avoid any conflicts in the cache which trigger checks/ensures.
-	GuidCache->RemoveNetGUID(Actor);
+	GuidCache->RemoveActorNetGUIDs(Actor);
 }
 
 bool UProxyNetDriver::ShouldCallRemoteFunction(UObject* Object, UFunction* Function, const FReplicationFlags& RepFlags) const
@@ -771,6 +772,8 @@ void UProxyNetDriver::InternalProcessRemoteFunction(
 
 int32 UProxyNetDriver::ServerReplicateActors(float DeltaSeconds)
 {
+	PrepareStateForRelevancy();
+
 	return Super::ServerReplicateActors(DeltaSeconds);
 }
 
@@ -1112,4 +1115,109 @@ FMultiServerProxyInternalConnectionRoute* UProxyNetDriver::GetRouteWithClientHan
 	}
 
 	return RoutePtr;
+}
+
+void UProxyNetDriver::PrepareStateForRelevancy()
+{
+	for (const TSharedPtr<FNetworkObjectInfo>& ObjectInfo : GetNetworkObjectList().GetAllObjects())
+	{
+		APlayerController* PlayerController = Cast<APlayerController>(ObjectInfo->Actor);
+		if (!PlayerController)
+		{
+			continue;
+		}
+
+		// The view target is used by the replication system to determine if an object is relevant to a given connection. The view
+		// target used by the game server is not replicated to the proxy so it must be manually set on the proxy in order to only
+		// replicate relevant objects to the proxy connections.
+		// 
+		// It's assumed, for now, that the view target for a player controller is the possessed pawn of the player controller that
+		// belongs to the route with the same proxy connection but is a route to the primary game server.
+		if (APawn* ViewTargetPawn = SetViewTargetUsingPrimaryPlayerController(PlayerController))
+		{
+			// Connections to non-primary game servers don't have an pawn so the view target must be manually set on the game server
+			// to ensure it replicates state nearby the pawn's position from another server.
+			if (ANoPawnPlayerController* NoPawnPlayerController = Cast<ANoPawnPlayerController>(PlayerController))
+			{
+				UpdateNoPawnPlayerControllerViewTarget(NoPawnPlayerController, ViewTargetPawn->GetReplicatedMovement().Location);
+			}
+		}
+	}
+
+	// Relevancy calcluations are made using the position returned by AActor::GetActorLocation(). Since the proxy is not simulating
+	// player movement the actor's location must be manually set to the location in AActor::ReplicatedMovement.
+	for (const TSharedPtr<FNetworkObjectInfo>& ObjectInfo : GetNetworkObjectList().GetAllObjects())
+	{
+		APawn* Pawn = Cast<APawn>(ObjectInfo->Actor);
+		if (!Pawn)
+		{
+			continue;
+		}
+
+		Pawn->SetActorLocation(Pawn->GetReplicatedMovement().Location);
+	}
+}
+
+APawn* UProxyNetDriver::SetViewTargetUsingPrimaryPlayerController(APlayerController* PlayerController)
+{
+	FMultiServerProxyInternalConnectionRoute* PlayerControllerRoute = GetRouteWithPlayerController(PlayerController);
+	if (PlayerControllerRoute == nullptr)
+	{
+		return nullptr;
+	}
+
+	// Find the player controller for a route that shares the same proxy connection but is considered the primary game server route.
+	TObjectPtr<const APlayerController> PrimaryPlayerController;
+	for (FMultiServerProxyInternalConnectionRoute& Route : InternalRoutes)
+	{
+		if (Route.ProxyConnection == PlayerControllerRoute->ProxyConnection && Route.State == EProxyConnectionState::ConnectedPrimary)
+		{
+			PrimaryPlayerController = Route.PlayerController;
+			break;
+		}
+	}
+
+	if (!PrimaryPlayerController)
+	{
+		UE_LOG(LogNetProxy, Verbose, TEXT("Cannot find a route for player controller %s that is connected to a primary game server."), *PlayerController->GetName());
+		return nullptr;
+	}
+
+	if (APawn* Pawn = PrimaryPlayerController->GetPawn())
+	{
+		UE_LOG(LogNetProxy, VeryVerbose, TEXT("Setting the view target of %s to %s using primary controller %s."),
+			*GetNameSafe(PlayerController),
+			*GetNameSafe(Pawn),
+			*GetNameSafe(PrimaryPlayerController));
+
+		PlayerController->SetViewTarget(Pawn);
+		return Pawn;
+	}
+
+	UE_LOG(LogNetProxy, Verbose, TEXT("Couldn't find a pawn for player controller %s using primary player controller %s."),
+		   *GetNameSafe(PlayerController),
+		   *GetNameSafe(PrimaryPlayerController));
+
+	return nullptr;
+}
+
+void UProxyNetDriver::UpdateNoPawnPlayerControllerViewTarget(ANoPawnPlayerController* PlayerController, FVector ViewTargetPos)
+{
+	FMultiServerProxyInternalConnectionRoute* PlayerControllerRoute = GetRouteWithPlayerController(PlayerController);
+	if (!ensure(PlayerControllerRoute))
+	{
+		return;
+	}
+
+	ensure(PlayerControllerRoute->State == EProxyConnectionState::Connected);
+	
+	UNetConnection* GameServerConnection = PlayerControllerRoute->GameServerConnection;
+
+	UE_LOG(LogNetProxy, VeryVerbose, TEXT("Updating view target position for %s to %s on connection %s."), 
+		   *PlayerController->GetName(), 
+		   *ViewTargetPos.ToString(),
+		   *GetNameSafe(GameServerConnection));
+
+	UFunction* Function = ANoPawnPlayerController::StaticClass()->FindFunctionByName(TEXT("ServerSetViewTargetPosition"));
+	GameServerConnection->Driver->ProcessRemoteFunction(PlayerController, Function, &ViewTargetPos, nullptr, nullptr);
 }
\ No newline at end of file
diff --git a/Engine/Plugins/Runtime/MultiServerReplication/Source/MultiServerReplication/Public/MultiServerProxy.h b/Engine/Plugins/Runtime/MultiServerReplication/Source/MultiServerReplication/Public/MultiServerProxy.h
--- a/Engine/Plugins/Runtime/MultiServerReplication/Source/MultiServerReplication/Public/MultiServerProxy.h
+++ b/Engine/Plugins/Runtime/MultiServerReplication/Source/MultiServerReplication/Public/MultiServerProxy.h
@@ -366,6 +366,15 @@ private:
 	/** Send any split join requests that have been queued up to the game server. */
 	void FlushSplitJoinRequests(FGameServerConnectionState* GameServerConnectionState);
 
+	/** Prepare state on the proxy to be used for relevancy when replicating to proxy connections. */
+	void PrepareStateForRelevancy();
+
+	/** Set the view target of a player controller. */
+	APawn* SetViewTargetUsingPrimaryPlayerController(APlayerController* PlayerController);
+
+	/** Update the view target position of a ANoPawnPlayerController on the game server. */
+	void UpdateNoPawnPlayerControllerViewTarget(class ANoPawnPlayerController* PlayerController, FVector ViewTargetPos);
+
 	/** Return a proxy connection from a primary player (a player on a primary game server). */
 	UNetConnection* GetProxyConnectionFromPrimaryPlayer(ULocalPlayer* Player);
 	
@@ -376,6 +385,7 @@ private:
 	FMultiServerProxyInternalConnectionRoute* GetRouteWithClientHandshakeId(uint32 ClientHandshakeId);
 
 	/** All of the routes for connections to game servers. */
+	UPROPERTY()
 	TArray<FMultiServerProxyInternalConnectionRoute> InternalRoutes;
 
 	/** Map a player to it's primary game server player controller. */
diff --git a/Engine/Source/Runtime/Engine/Classes/Engine/PackageMapClient.h b/Engine/Source/Runtime/Engine/Classes/Engine/PackageMapClient.h
--- a/Engine/Source/Runtime/Engine/Classes/Engine/PackageMapClient.h
+++ b/Engine/Source/Runtime/Engine/Classes/Engine/PackageMapClient.h
@@ -252,6 +252,9 @@ public:
 
 	ENGINE_API bool 		RemoveNetGUID(const UObject* Object);
 
+	/** Remove an actor, it's components and all sub-object NetGUIDs from the NetGUID cache. */
+	ENGINE_API void 		RemoveActorNetGUIDs(const AActor* Actor);
+
 	ENGINE_API void			AsyncPackageCallback(const FName& PackageName, UPackage * Package, EAsyncLoadingResult::Type Result);
 
 	ENGINE_API void			ResetCacheForDemo();
diff --git a/Engine/Source/Runtime/Engine/Classes/GameFramework/PlayerController.h b/Engine/Source/Runtime/Engine/Classes/GameFramework/PlayerController.h
--- a/Engine/Source/Runtime/Engine/Classes/GameFramework/PlayerController.h
+++ b/Engine/Source/Runtime/Engine/Classes/GameFramework/PlayerController.h
@@ -2346,4 +2346,12 @@ class ANoPawnPlayerController : public APlayerController
 public:
 
 	GENERATED_BODY()
+
+	virtual void PostInitializeComponents() override;
+	virtual class AActor* GetViewTarget() const override;
+	virtual class UNetConnection* GetNetConnection() const override;
+
+	/** Set a view target position for use in replication relevancy. */
+	UFUNCTION(Server, Unreliable)
+	void ServerSetViewTargetPosition(FVector ViewTargetPos);
 };
diff --git a/Engine/Source/Runtime/Engine/Private/Net/NetSubObjectRegistryGetter.h b/Engine/Source/Runtime/Engine/Private/Net/NetSubObjectRegistryGetter.h
--- a/Engine/Source/Runtime/Engine/Private/Net/NetSubObjectRegistryGetter.h
+++ b/Engine/Source/Runtime/Engine/Private/Net/NetSubObjectRegistryGetter.h
@@ -20,14 +20,14 @@ public:
 
 	using FSubObjectRegistry = UE::Net::FSubObjectRegistry;
 
-	static const FSubObjectRegistry& GetSubObjects(AActor* InActor)
+	static const FSubObjectRegistry& GetSubObjects(const AActor* InActor)
 	{
 		return InActor->ReplicatedSubObjects;
 	}
 
-	static const FSubObjectRegistry* GetSubObjectsOfActorComponent(AActor* InActor, UActorComponent* InActorComp)
+	static const FSubObjectRegistry* GetSubObjectsOfActorComponent(const AActor* InActor, UActorComponent* InActorComp)
 	{
-		UE::Net::FReplicatedComponentInfo* ComponentInfo = InActor->ReplicatedComponentsInfo.FindByKey(InActorComp);
+		const UE::Net::FReplicatedComponentInfo* ComponentInfo = InActor->ReplicatedComponentsInfo.FindByKey(InActorComp);
 		return ComponentInfo ? &(ComponentInfo->SubObjects) : nullptr;
 	}
 
diff --git a/Engine/Source/Runtime/Engine/Private/PackageMapClient.cpp b/Engine/Source/Runtime/Engine/Private/PackageMapClient.cpp
--- a/Engine/Source/Runtime/Engine/Private/PackageMapClient.cpp
+++ b/Engine/Source/Runtime/Engine/Private/PackageMapClient.cpp
@@ -21,6 +21,7 @@
 #include "Serialization/MemoryReader.h"
 #include "Net/NetworkGranularMemoryLogging.h"
 #include "Misc/CommandLine.h"
+#include "Net/NetSubObjectRegistryGetter.h"
 
 #if UE_WITH_IRIS
 #include "Iris/ReplicationSystem/NetTokenStore.h"
@@ -4201,6 +4202,31 @@ bool FNetGUIDCache::RemoveNetGUID(const UObject* Object)
 	return false;
 }
 
+void FNetGUIDCache::RemoveActorNetGUIDs(const AActor* Actor)
+{
+	using namespace UE::Net;
+
+	auto RemoveSubObjectsFromGuidCache = [this](const FSubObjectRegistry* SubObjectRegistry)
+	{
+		if (SubObjectRegistry)
+		{
+			for (const FSubObjectRegistry::FEntry& SubObjectInfo : SubObjectRegistry->GetRegistryList())
+			{
+				RemoveNetGUID(SubObjectInfo.GetSubObject());
+			}
+		}
+	};
+
+	RemoveNetGUID(Actor);
+	RemoveSubObjectsFromGuidCache(&FSubObjectRegistryGetter::GetSubObjects(Actor));
+
+	for (UActorComponent* Component : Actor->GetComponents())
+	{
+		RemoveNetGUID(Component);
+		RemoveSubObjectsFromGuidCache(FSubObjectRegistryGetter::GetSubObjectsOfActorComponent(Actor, Component));
+	}
+}
+
 void FNetGUIDCache::GenerateFullNetGUIDPath_r( const FNetworkGUID& NetGUID, FString& FullPath ) const
 {
 	if ( !NetGUID.IsValid() )
diff --git a/Engine/Source/Runtime/Engine/Private/PlayerController.cpp b/Engine/Source/Runtime/Engine/Private/PlayerController.cpp
--- a/Engine/Source/Runtime/Engine/Private/PlayerController.cpp
+++ b/Engine/Source/Runtime/Engine/Private/PlayerController.cpp
@@ -6643,5 +6643,29 @@ void APlayerController::ClientAckTimeDilation_Implementation(float TimeDilation,
 	}
 }
 
+void ANoPawnPlayerController::PostInitializeComponents()
+{
+	// Don't call APlayerController::PostInitializeComponents() because it will spawn an instance of APlayerState
+	// which doesn't make sense in ANoPawnPlayerController because it doesn't have an associated pawn.
+
+	AController::PostInitializeComponents();
+}
+
+AActor* ANoPawnPlayerController::GetViewTarget() const
+{
+	return const_cast<ANoPawnPlayerController*>(this);
+}
+
+class UNetConnection* ANoPawnPlayerController::GetNetConnection() const
+{
+	return NetConnection;
+}
+
+void ANoPawnPlayerController::ServerSetViewTargetPosition_Implementation(FVector ViewTargetPos)
+{
+	// Set the value of FRepMovement because SetActorLocation() is marked as private.
+	GetReplicatedMovement_Mutable().Location = ViewTargetPos;
+}
+
 #undef LOCTEXT_NAMESPACE
 
diff --git a/Engine/Source/Runtime/Engine/Public/Net/DataChannel.h b/Engine/Source/Runtime/Engine/Public/Net/DataChannel.h
--- a/Engine/Source/Runtime/Engine/Public/Net/DataChannel.h
+++ b/Engine/Source/Runtime/Engine/Public/Net/DataChannel.h
@@ -225,6 +225,9 @@ DEFINE_CONTROL_CHANNEL_MESSAGE(IrisNetRefHandleErrorWithArray, 33, UE::Net::ENet
 //	* Spawns ANoPawnPlayerController instead of using the class in AGameMode::PlayerControllerClass.
 //	* Network relevency will be based on the location of ANoPawnPlayerController.
 //
+// Relevancy: Since network relevancy will be based on the location of ANoPawnPlayerController, it's important to set the location of
+//			  the actor using the RPC ANoPawnPlayerController::SetViewTargetPosition(), either from the server or client using this connection.
+//
 // These messages are only intended to be used when called from one game server to another game server. For this reason it is disabled
 // by default and must be explicitly enabled on the receiving server using the cvar net.EnableSupportForNoPawnConnection or calling
 // UNetDriver::SetSupportForNoPawnConnection() directly.
