Limit the rate at which the view target of non-primary game server connections is updated.

diff --git a/Engine/Plugins/Runtime/MultiServerReplication/Source/MultiServerReplication/Private/MultiServerProxy.cpp b/Engine/Plugins/Runtime/MultiServerReplication/Source/MultiServerReplication/Private/MultiServerProxy.cpp
--- a/Engine/Plugins/Runtime/MultiServerReplication/Source/MultiServerReplication/Private/MultiServerProxy.cpp
+++ b/Engine/Plugins/Runtime/MultiServerReplication/Source/MultiServerReplication/Private/MultiServerProxy.cpp
@@ -37,6 +37,12 @@ namespace UE::MultiServerProxy::Private
 			return FString::Printf(TEXT("None:None"));
 		}
 	}
+
+	float NonPrimarySetViewTargetIntervalSec = 1.0;
+	static FAutoConsoleVariableRef CVarNonPrimarySetViewTargetIntervalSec(
+		TEXT("net.proxy.NonPrimarySetViewTargetInterval"),
+		NonPrimarySetViewTargetIntervalSec,
+		TEXT("The interval, in seconds, between updating the view target position on non-primary game servers."));
 }
 
 FProxyNetGUIDCache::FProxyNetGUIDCache(UNetDriver* NetDriver) : FNetGUIDCache(NetDriver)
@@ -326,6 +332,8 @@ void UProxyListenerNotify::ConnectToGameServer(UNetConnection* ProxyConnection,
 	Route.GameServerConnection = nullptr;
 	Route.Player = NewPlayer;
 	Route.PlayerController = nullptr;
+	Route.LastUpdateViewTargetSec = 0.0;
+	Route.LastViewTargetPos = FVector(0, 0, 0);
 
 	if (!EnumHasAllFlags(Flags, UE::Net::EJoinFlags::NoPawn))
 	{
@@ -1210,6 +1218,21 @@ void UProxyNetDriver::UpdateNoPawnPlayerControllerViewTarget(ANoPawnPlayerContro
 	}
 
 	ensure(PlayerControllerRoute->State == EProxyConnectionState::Connected);
+
+	if (PlayerControllerRoute->LastViewTargetPos.Equals(ViewTargetPos))
+	{
+		return;
+	}
+	
+	const double WorldTimeSec = World->GetTimeSeconds();
+	const double NextUpdateTime = (PlayerControllerRoute->LastUpdateViewTargetSec + UE::MultiServerProxy::Private::NonPrimarySetViewTargetIntervalSec);
+	if (NextUpdateTime >= WorldTimeSec)
+	{
+		return;
+	}
+
+	PlayerControllerRoute->LastUpdateViewTargetSec = WorldTimeSec;
+	PlayerControllerRoute->LastViewTargetPos = ViewTargetPos;
 	
 	UNetConnection* GameServerConnection = PlayerControllerRoute->GameServerConnection;
 
diff --git a/Engine/Plugins/Runtime/MultiServerReplication/Source/MultiServerReplication/Public/MultiServerProxy.h b/Engine/Plugins/Runtime/MultiServerReplication/Source/MultiServerReplication/Public/MultiServerProxy.h
--- a/Engine/Plugins/Runtime/MultiServerReplication/Source/MultiServerReplication/Public/MultiServerProxy.h
+++ b/Engine/Plugins/Runtime/MultiServerReplication/Source/MultiServerReplication/Public/MultiServerProxy.h
@@ -137,6 +137,9 @@ struct FMultiServerProxyInternalConnectionRoute
 
 	/** The state of the route. */
 	mutable EProxyConnectionState State = EProxyConnectionState::Disconnected;
+	
+	mutable double LastUpdateViewTargetSec;
+	mutable FVector LastViewTargetPos;
 };
 
 /** 
