MultiServer proxy algorithmic optimizations for route lookup.

diff --git a/Engine/Plugins/Runtime/MultiServerReplication/Source/MultiServerReplication/Private/MultiServerProxy.cpp b/Engine/Plugins/Runtime/MultiServerReplication/Source/MultiServerReplication/Private/MultiServerProxy.cpp
--- a/Engine/Plugins/Runtime/MultiServerReplication/Source/MultiServerReplication/Private/MultiServerProxy.cpp
+++ b/Engine/Plugins/Runtime/MultiServerReplication/Source/MultiServerReplication/Private/MultiServerProxy.cpp
@@ -344,7 +344,7 @@ void UProxyListenerNotify::ConnectToGameServer(UNetConnection* ProxyConnection,
 		Route.State = EProxyConnectionState::Connecting;
 	}
 	
-	ProxyNetDriver->InternalRoutes.Add(Route);
+	ProxyNetDriver->InternalRoutes.Add(ClientHandshakeId, Route);
 }
 
 void UProxyBackendNetConnection::HandleClientPlayer(APlayerController* NewPlayerController, UNetConnection* GameServerConnection)
@@ -419,7 +419,7 @@ void UProxyBackendNetDriver::ForwardRemoteFunction(UObject* RootObject, UObject*
 	// Drop RPCs that don't come from a player controller on a primary game server.
 	if (APlayerController* PlayerController = Cast<APlayerController>(OwningActor))
 	{
-		FMultiServerProxyInternalConnectionRoute* Route = ProxyNetDriver->GetRouteWithPlayerController(PlayerController);
+		const FMultiServerProxyInternalConnectionRoute* Route = ProxyNetDriver->GetRouteWithPlayerController(PlayerController);
 		if (!ensure(Route))
 		{
 			return;
@@ -852,7 +852,7 @@ void UProxyNetDriver::GameServerAssignPlayerController(UNetConnection* ChildGame
 		   *ChildGameServerConnection->GetDriver()->GetName(), *ChildGameServerConnection->GetName());
 
 	const uint32 ClientHandshakeId = GameServerPlayerController->GetClientHandshakeId();
-	FMultiServerProxyInternalConnectionRoute* Route = GetRouteWithClientHandshakeId(ClientHandshakeId);
+	FMultiServerProxyInternalConnectionRoute* Route = InternalRoutes.Find(ClientHandshakeId);
 	if (!ensure(Route))
 	{
 		return;
@@ -928,34 +928,18 @@ void UProxyNetDriver::ReassignPlayerController(const FMultiServerProxyInternalCo
 	ensure(Route->GameServerConnection->GetClientHandshakeId() == PlayerController->GetClientHandshakeId());
 	ensure(Route->GameServerConnection == GameServerConnection);
 
-	FMultiServerProxyInternalConnectionRoute* PrevRoute = GetRouteWithPlayerController(PlayerController);
-	if (PrevRoute)
-	{
-		UE_LOG(LogNetProxy, Log, TEXT("[%u] Clearing the player controller %s from the previous route with client handshake id %u"), 
-			   Route->ClientHandshakeId,
-			   *UE::MultiServerProxy::Private::PlayerControllerToString(PlayerController),
-			   PrevRoute->ClientHandshakeId);
-
-		
-		// When a player controller is assigned to a new route it's assumed that the route that is currently using that player controller will
-		// also have a new controller assigned. The PendingReassign state marks that other route as in a pending state that is waiting for it's
-		// new controller.
-		PrevRoute->PlayerController = nullptr;
-		PrevRoute->State = EProxyConnectionState::PendingReassign;
-	}
-	else
-	{
-		UE_LOG(LogNetProxy, Log, TEXT("[%u] Player controller %s was not bound to another route."), 
-			   Route->ClientHandshakeId,
-			   *UE::MultiServerProxy::Private::PlayerControllerToString(PlayerController));
-	}
-
 	const APlayerController* PrevPlayerController = Route->PlayerController;
 
 	PlayerController->NetConnection = GameServerConnection;
 	PlayerController->Player = Route->Player;
 	Route->PlayerController = PlayerController;
 
+	// The proxy listens for 2 events when a player controller migrates from one game server to another:
+	//	* The player controller on the connection to game server A will change to an instance of ANoPawnPlayerController.
+	//	* The player controller on the connection to game server B will change to the instance of the game's player controller.
+	//
+	// IMPORTANT: The code below must not assume that these events will arrive in a specific order, or even arrive at all.
+
 	if (PlayerController->IsA(ANoPawnPlayerController::StaticClass()))
 	{
 		UE_LOG(LogNetProxy, Log, TEXT("[%u] Deregistering player %s as a primary game server."), 
@@ -973,6 +957,7 @@ void UProxyNetDriver::ReassignPlayerController(const FMultiServerProxyInternalCo
 			   *UE::MultiServerProxy::Private::PlayerControllerToString(Route->PlayerController));
 
 		PlayerToPrimaryGameServerPlayerController.Add(Route->Player, Route->PlayerController);
+		ProxyConnectionToPrimaryGameServerPlayerController.Add(Route->ProxyConnection, Route->PlayerController);
 		Route->State = EProxyConnectionState::ConnectedPrimary;
 	}
 
@@ -1017,6 +1002,7 @@ void UProxyNetDriver::ProcessConnectionState_ConnectingPrimary(FMultiServerProxy
 		   *UE::MultiServerProxy::Private::PlayerControllerToString(Route->PlayerController));
 
 	PlayerToPrimaryGameServerPlayerController.Add(Route->Player, Route->PlayerController);
+	ProxyConnectionToPrimaryGameServerPlayerController.Add(Route->ProxyConnection, Route->PlayerController);
 
 	Route->State = EProxyConnectionState::ConnectedPrimary;
 
@@ -1083,7 +1069,7 @@ UNetConnection* UProxyNetDriver::GetProxyConnectionFromPrimaryPlayer(ULocalPlaye
 
 	const APlayerController* GameServerPlayerController = GameServerPlayerControllerPtr->Get();
 
-	FMultiServerProxyInternalConnectionRoute* Route = GetRouteWithPlayerController(GameServerPlayerController);
+	const FMultiServerProxyInternalConnectionRoute* Route = GetRouteWithPlayerController(GameServerPlayerController);
 	if (!Route)
 	{
 		return nullptr;
@@ -1095,61 +1081,45 @@ UNetConnection* UProxyNetDriver::GetProxyConnectionFromPrimaryPlayer(ULocalPlaye
 	return Route->ProxyConnection;
 }
 
-FMultiServerProxyInternalConnectionRoute* UProxyNetDriver::GetRouteWithPlayerController(const APlayerController* PlayerController)
+const FMultiServerProxyInternalConnectionRoute* UProxyNetDriver::GetRouteWithPlayerController(const APlayerController* PlayerController)
 {
-	FMultiServerProxyInternalConnectionRoute* RoutePtr = nullptr;
-	for (FMultiServerProxyInternalConnectionRoute& Route : InternalRoutes)
-	{
-		if (Route.PlayerController == PlayerController)
-		{
-			ensure(RoutePtr == nullptr);
-			RoutePtr = &Route;
-		}
-	}
-
-	return RoutePtr;
+	return InternalRoutes.Find(PlayerController->GetClientHandshakeId());
 }
 
-FMultiServerProxyInternalConnectionRoute* UProxyNetDriver::GetRouteWithClientHandshakeId(uint32 ClientHandshakeId)
+void UProxyNetDriver::PrepareStateForRelevancy()
 {
-	FMultiServerProxyInternalConnectionRoute* RoutePtr = nullptr;
-	for (FMultiServerProxyInternalConnectionRoute& Route : InternalRoutes)
+	for (TPair<TObjectPtr<UNetConnection>, TObjectPtr<const APlayerController>>& ProxyConnectionToPlayerController : ProxyConnectionToPrimaryGameServerPlayerController)
 	{
-		if (Route.ClientHandshakeId == ClientHandshakeId)
-		{
-			ensure(RoutePtr == nullptr);
-			RoutePtr = &Route;
-		}
-	}
+		const APlayerController* PrimaryPlayerController = ProxyConnectionToPlayerController.Value;
 
-	return RoutePtr;
-}
+		ensure(!PrimaryPlayerController->IsA(ANoPawnPlayerController::StaticClass()));
 
-void UProxyNetDriver::PrepareStateForRelevancy()
-{
-	for (const TSharedPtr<FNetworkObjectInfo>& ObjectInfo : GetNetworkObjectList().GetAllObjects())
+		UpdateLocalViewTarget(PrimaryPlayerController);
+	}
+
+	for (TPair<uint32, FMultiServerProxyInternalConnectionRoute>& ClientHandshakeIdToRoute : InternalRoutes)
 	{
-		APlayerController* PlayerController = Cast<APlayerController>(ObjectInfo->Actor);
-		if (!PlayerController)
+		const FMultiServerProxyInternalConnectionRoute* Route = &ClientHandshakeIdToRoute.Value;
+
+		const ANoPawnPlayerController* NoPawnPlayerController = Cast<ANoPawnPlayerController>(Route->PlayerController);
+		if (!NoPawnPlayerController)
 		{
 			continue;
 		}
 
-		// The view target is used by the replication system to determine if an object is relevant to a given connection. The view
-		// target used by the game server is not replicated to the proxy so it must be manually set on the proxy in order to only
-		// replicate relevant objects to the proxy connections.
-		// 
-		// It's assumed, for now, that the view target for a player controller is the possessed pawn of the player controller that
-		// belongs to the route with the same proxy connection but is a route to the primary game server.
-		if (APawn* ViewTargetPawn = SetViewTargetUsingPrimaryPlayerController(PlayerController))
+		TObjectPtr<const APlayerController>* PrimaryPlayerControllerPtr = ProxyConnectionToPrimaryGameServerPlayerController.Find(Route->ProxyConnection);
+		if (!PrimaryPlayerControllerPtr)
 		{
-			// Connections to non-primary game servers don't have an pawn so the view target must be manually set on the game server
-			// to ensure it replicates state nearby the pawn's position from another server.
-			if (ANoPawnPlayerController* NoPawnPlayerController = Cast<ANoPawnPlayerController>(PlayerController))
-			{
-				UpdateNoPawnPlayerControllerViewTarget(NoPawnPlayerController, ViewTargetPawn->GetReplicatedMovement().Location);
-			}
+			continue;
 		}
+
+		const APlayerController* PrimaryPlayerController = *PrimaryPlayerControllerPtr;
+
+		FVector Location;
+		FRotator Rotation;
+		PrimaryPlayerController->GetPlayerViewPoint(Location, Rotation);
+
+		SetRemoteViewTarget(NoPawnPlayerController, Location);
 	}
 
 	// Relevancy calcluations are made using the position returned by AActor::GetActorLocation(). Since the proxy is not simulating
@@ -1166,52 +1136,26 @@ void UProxyNetDriver::PrepareStateForRelevancy()
 	}
 }
 
-APawn* UProxyNetDriver::SetViewTargetUsingPrimaryPlayerController(APlayerController* PlayerController)
+void UProxyNetDriver::UpdateLocalViewTarget(const APlayerController* PrimaryPlayerController)
 {
-	FMultiServerProxyInternalConnectionRoute* PlayerControllerRoute = GetRouteWithPlayerController(PlayerController);
-	if (PlayerControllerRoute == nullptr)
-	{
-		return nullptr;
-	}
-
-	// Find the player controller for a route that shares the same proxy connection but is considered the primary game server route.
-	TObjectPtr<const APlayerController> PrimaryPlayerController;
-	for (FMultiServerProxyInternalConnectionRoute& Route : InternalRoutes)
-	{
-		if (Route.ProxyConnection == PlayerControllerRoute->ProxyConnection && Route.State == EProxyConnectionState::ConnectedPrimary)
-		{
-			PrimaryPlayerController = Route.PlayerController;
-			break;
-		}
-	}
-
-	if (!PrimaryPlayerController)
-	{
-		UE_LOG(LogNetProxy, Verbose, TEXT("Cannot find a route for player controller %s that is connected to a primary game server."), *PlayerController->GetName());
-		return nullptr;
-	}
-
 	if (APawn* Pawn = PrimaryPlayerController->GetPawn())
 	{
-		UE_LOG(LogNetProxy, VeryVerbose, TEXT("Setting the view target of %s to %s using primary controller %s."),
-			*GetNameSafe(PlayerController),
+		UE_LOG(LogNetProxy, VeryVerbose, TEXT("Setting the view target of %s to %s and updating the camera manager."),
 			*GetNameSafe(Pawn),
 			*GetNameSafe(PrimaryPlayerController));
 
-		PlayerController->SetViewTarget(Pawn);
-		return Pawn;
+	
+		const_cast<APlayerController*>(PrimaryPlayerController)->SetViewTarget(Pawn);
+	}
+	else
+	{
+		UE_LOG(LogNetProxy, Verbose, TEXT("Couldn't find a pawn for player controller %s."), *GetNameSafe(PrimaryPlayerController));
 	}
-
-	UE_LOG(LogNetProxy, Verbose, TEXT("Couldn't find a pawn for player controller %s using primary player controller %s."),
-		   *GetNameSafe(PlayerController),
-		   *GetNameSafe(PrimaryPlayerController));
-
-	return nullptr;
 }
 
-void UProxyNetDriver::UpdateNoPawnPlayerControllerViewTarget(ANoPawnPlayerController* PlayerController, FVector ViewTargetPos)
+void UProxyNetDriver::SetRemoteViewTarget(const ANoPawnPlayerController* PlayerController, FVector ViewTargetPos)
 {
-	FMultiServerProxyInternalConnectionRoute* PlayerControllerRoute = GetRouteWithPlayerController(PlayerController);
+	const FMultiServerProxyInternalConnectionRoute* PlayerControllerRoute = GetRouteWithPlayerController(PlayerController);
 	if (!ensure(PlayerControllerRoute))
 	{
 		return;
@@ -1242,5 +1186,5 @@ void UProxyNetDriver::UpdateNoPawnPlayerControllerViewTarget(ANoPawnPlayerContro
 		   *GetNameSafe(GameServerConnection));
 
 	UFunction* Function = ANoPawnPlayerController::StaticClass()->FindFunctionByName(TEXT("ServerSetViewTargetPosition"));
-	GameServerConnection->Driver->ProcessRemoteFunction(PlayerController, Function, &ViewTargetPos, nullptr, nullptr);
+	GameServerConnection->Driver->ProcessRemoteFunction(const_cast<ANoPawnPlayerController*>(PlayerController), Function, &ViewTargetPos, nullptr, nullptr);
 }
\ No newline at end of file
diff --git a/Engine/Plugins/Runtime/MultiServerReplication/Source/MultiServerReplication/Public/MultiServerProxy.h b/Engine/Plugins/Runtime/MultiServerReplication/Source/MultiServerReplication/Public/MultiServerProxy.h
--- a/Engine/Plugins/Runtime/MultiServerReplication/Source/MultiServerReplication/Public/MultiServerProxy.h
+++ b/Engine/Plugins/Runtime/MultiServerReplication/Source/MultiServerReplication/Public/MultiServerProxy.h
@@ -125,13 +125,13 @@ struct FMultiServerProxyInternalConnectionRoute
 	UPROPERTY()
 	TObjectPtr<class ULocalPlayer> Player;
 
-	/** 
-	 * The functions below are marked as mutable to allow the rest of the structure to remain const
-	 * after the initial connection has been finalized, avoiding mistakenly changing route state 
-	 * when migrating player controllers between game servers and changing primary game server.
+	/**
+	 * The player controller.
+	 * 
+	 * The pointer to APlayerController is const in order to make sure changes to the pointer are
+	 * explicit using const_cast. This is important because accidental changes to the controller, 
+	 * especially when being reassigned, can lead to difficult to find routing bugs.
 	 */
-
-	/** The player controller. */
 	UPROPERTY()
 	mutable TObjectPtr<const class APlayerController> PlayerController;
 
@@ -372,29 +372,36 @@ private:
 	/** Prepare state on the proxy to be used for relevancy when replicating to proxy connections. */
 	void PrepareStateForRelevancy();
 
-	/** Set the view target of a player controller. */
-	APawn* SetViewTargetUsingPrimaryPlayerController(APlayerController* PlayerController);
+	/** Update the view target on player controllers used by proxy client connections. */
+	void UpdateLocalViewTarget(const APlayerController* PrimaryPlayerController);
 
-	/** Update the view target position of a ANoPawnPlayerController on the game server. */
-	void UpdateNoPawnPlayerControllerViewTarget(class ANoPawnPlayerController* PlayerController, FVector ViewTargetPos);
+	/** Set the view target on no pawn connections to the game server. */
+	void SetRemoteViewTarget(const class ANoPawnPlayerController* PlayerController, FVector ViewTargetPos);
 
 	/** Return a proxy connection from a primary player (a player on a primary game server). */
 	UNetConnection* GetProxyConnectionFromPrimaryPlayer(ULocalPlayer* Player);
 	
 	/** Returns the internal route that uses a player controller. */
-	FMultiServerProxyInternalConnectionRoute* GetRouteWithPlayerController(const APlayerController* PlayerController);
-
-	/** Returns the internal route that uses a client handshake id. */
-	FMultiServerProxyInternalConnectionRoute* GetRouteWithClientHandshakeId(uint32 ClientHandshakeId);
-
-	/** All of the routes for connections to game servers. */
+	const FMultiServerProxyInternalConnectionRoute* GetRouteWithPlayerController(const APlayerController* PlayerController);
+
+	/**
+	 * Map client handshake ids to a route data structure.
+	 * 	
+	 * The size of this map is (m*n), where m is the number of proxy connections and n is the number of 
+	 * game servers. For this reason enumerating of this map should be avoided and rather use the functions
+	 * GetRouteWithPlayerController(), GetRouteWithClientHandshakeId(), and other internal map data structures.
+	 */
 	UPROPERTY()
-	TArray<FMultiServerProxyInternalConnectionRoute> InternalRoutes;
+	TMap<uint32, FMultiServerProxyInternalConnectionRoute> InternalRoutes;
 
 	/** Map a player to it's primary game server player controller. */
 	UPROPERTY()
 	TMap<TObjectPtr<class ULocalPlayer>, TObjectPtr<const APlayerController>> PlayerToPrimaryGameServerPlayerController;
 
+	/** Map a proxy connection to the primary game server player controller. */
+	UPROPERTY()
+	TMap<TObjectPtr<UNetConnection>, TObjectPtr<const APlayerController>> ProxyConnectionToPrimaryGameServerPlayerController;
+
 	/** Net drivers and associated state used to connect to backend game servers. */
 	UPROPERTY()
 	TArray<FGameServerConnectionState> GameServerConnections;
