Added metrics to measure the performance and behavior of the multi-server proxy.

diff --git a/Engine/Config/BaseEngine.ini b/Engine/Config/BaseEngine.ini
--- a/Engine/Config/BaseEngine.ini
+++ b/Engine/Config/BaseEngine.ini
@@ -3738,6 +3738,9 @@ bStaticMeshUseSmoothEdgesIfSmoothingInformationIsMissing=True
 GenericPipelineClass=/Script/InterchangePipelines.InterchangeGenericAssetsPipeline
 ConverterDefaultPipeline=/Interchange/Pipelines/DefaultAssetsPipeline.DefaultAssetsPipeline
 
+[/Script/Engine.NetworkMetricsLog]
+IntervalSeconds=10
+
 [/Script/Engine.NetworkMetricsCSV_Replication]
 IntervalSeconds=0
 
diff --git a/Engine/Plugins/Runtime/MultiServerReplication/Source/MultiServerReplication/Private/MultiServerProxy.cpp b/Engine/Plugins/Runtime/MultiServerReplication/Source/MultiServerReplication/Private/MultiServerProxy.cpp
--- a/Engine/Plugins/Runtime/MultiServerReplication/Source/MultiServerReplication/Private/MultiServerProxy.cpp
+++ b/Engine/Plugins/Runtime/MultiServerReplication/Source/MultiServerReplication/Private/MultiServerProxy.cpp
@@ -9,6 +9,7 @@
 #include "Misc/CommandLine.h"
 #include "Net/DataChannel.h"
 #include "Net/Core/Misc/NetSubObjectRegistry.h"
+#include "Net/NetworkMetricsDefs.h"
 
 DEFINE_LOG_CATEGORY_STATIC(LogNetProxy, Log, All);
 
@@ -45,6 +46,32 @@ namespace UE::MultiServerProxy::Private
 		TEXT("The interval, in seconds, between updating the view target position on non-primary game servers."));
 }
 
+namespace UE::MultiServerProxy::Metrics
+{
+	// The frame rate (FPS) and associated frame time in milliseconds of the process.
+	const FName FPS("FPS");
+	const FName FrameTimeMS("FrameTimeMS");
+	
+	// The number of RPCs forwarded from the proxy client to a primary game server.
+	const FName NumForwardClientRPC("NumForwardClientRPC");
+
+	// The number of RPCs forwarded from a game server to a proxy connection.
+	const FName NumForwardBackendRPC("NumForwardBackendRPC");
+	
+	// The size of the UProxyNetDriver::InternalRoutes array which represents a path from each proxy
+	// client to each game server.
+	const FName NumProxyRoutes("NumProxyRoutes");
+	
+	// Aggregated InRate from all NetDrivers used to connect to a backend game server.
+	const FName AggregatedBackendInRate("AggregatedBackendInRate");
+	
+	// Aggregated OutRate from all NetDrivers used to connect to a backend game server.
+	const FName AggregatedBackendOutRate("AggregatedBackendOutRate");
+
+	// The number of player controller's re-assigned on backend game servers.
+	const FName NumReassignPlayerController("NumReassignPlayerController");
+}
+
 FProxyNetGUIDCache::FProxyNetGUIDCache(UNetDriver* NetDriver) : FNetGUIDCache(NetDriver)
 {
 
@@ -409,6 +436,8 @@ void UProxyBackendNetDriver::ForwardRemoteFunction(UObject* RootObject, UObject*
 	// This function is called when the proxy receives an RPC from a game server and will only forward the function on
 	// to the frontend net driver if it is owned by a player that considers that game server the primary game server.
 
+	ProxyNetDriver->GetMetrics()->IncrementInt(UE::MultiServerProxy::Metrics::NumForwardClientRPC, 1);
+
 	AActor* OwningActor = Cast<AActor>(RootObject);
 	if (!ensure(OwningActor != nullptr))
 	{
@@ -550,6 +579,19 @@ bool UProxyBackendNetDriver::ShouldClientDestroyActor(AActor* Actor, EChannelClo
 	return (CloseReason != EChannelCloseReason::Migrated);
 }
 
+bool UProxyBackendNetDriver::ShouldUpdateStats() const
+{
+	return true;
+}
+
+bool UProxyBackendNetDriver::ShouldRegisterMetricsDatabaseListeners() const
+{
+	// Metrics database listeners are not used by UProxyBackendNetDriver because they are aggregated in UProxyNetDriver 
+	// and reported to the metrics database listeners.
+
+	return false;
+}
+
 void UProxyNetDriver::RegisterGameServer(const FURL& GameServerURL)
 {
 	UE_LOG(LogNetProxy, Log, TEXT("Registering Game Server for URL: %s"), *GameServerURL.ToString());
@@ -637,6 +679,20 @@ bool UProxyNetDriver::InitBase(bool bInitAsClient, FNetworkNotify* InNotify, con
 	return bSuccess;
 }
 
+void UProxyNetDriver::SetupNetworkMetrics()
+{
+	Super::SetupNetworkMetrics();
+	
+	GetMetrics()->CreateFloat(UE::MultiServerProxy::Metrics::FrameTimeMS, 0.0f);
+	GetMetrics()->CreateFloat(UE::MultiServerProxy::Metrics::FPS, 0.0f);
+	GetMetrics()->CreateInt(UE::MultiServerProxy::Metrics::NumForwardClientRPC, 0);
+	GetMetrics()->CreateInt(UE::MultiServerProxy::Metrics::NumForwardBackendRPC, 0);
+	GetMetrics()->CreateInt(UE::MultiServerProxy::Metrics::NumProxyRoutes, 0);
+	GetMetrics()->CreateInt(UE::MultiServerProxy::Metrics::AggregatedBackendInRate, 0);
+	GetMetrics()->CreateInt(UE::MultiServerProxy::Metrics::AggregatedBackendOutRate, 0);
+	GetMetrics()->CreateInt(UE::MultiServerProxy::Metrics::NumReassignPlayerController, 0);
+}
+
 bool UProxyNetDriver::InitConnect(class FNetworkNotify* InNotify, const FURL& ConnectURL, FString& Error)
 {
 	checkf(0, TEXT("UProxyNetDriver is only intended to be used to receive connections and not establish outgoing connections."));
@@ -648,6 +704,8 @@ void UProxyNetDriver::ForwardRemoteFunction(UObject* RootObject, UObject* SubObj
 	// This function is called when the proxy receives an RPC from a game client and will forward the function on
 	// to the owning player's primary game server.
 
+	GetMetrics()->IncrementInt(UE::MultiServerProxy::Metrics::NumForwardBackendRPC, 1);
+
 	// If the owner is a PlayerController it will be for a proxy connection. There is no need to map it to the PlayerController
 	// on the game server because UNetDriver::ProcessRemoteFunction() will automatically send the RPC to the server connection
 	// and ignores the value of AActor::GetNetConnection().
@@ -821,6 +879,18 @@ void UProxyNetDriver::Shutdown()
 	Super::Shutdown();
 }
 
+void UProxyNetDriver::TickFlush(float DeltaSeconds)
+{
+	// Update any metrics in the database that aren't tracked while ticking the frame.
+	UpdateProxyNetworkMetrics();
+
+	// The network stats and network metrics database are updated and reported to listeners here.
+	Super::TickFlush(DeltaSeconds);
+
+	// Reset any metrics that are tracked while ticking the frame.
+	ResetProxyNetworkMetrics();
+}
+
 int32 UProxyNetDriver::GetNextGameServerClientId()
 {
 	return NextGameServerClientId++;
@@ -928,6 +998,8 @@ void UProxyNetDriver::ReassignPlayerController(const FMultiServerProxyInternalCo
 	ensure(Route->GameServerConnection->GetClientHandshakeId() == PlayerController->GetClientHandshakeId());
 	ensure(Route->GameServerConnection == GameServerConnection);
 
+	GetMetrics()->IncrementInt(UE::MultiServerProxy::Metrics::NumReassignPlayerController, 1);
+
 	const APlayerController* PrevPlayerController = Route->PlayerController;
 
 	PlayerController->NetConnection = GameServerConnection;
@@ -1029,6 +1101,69 @@ void UProxyNetDriver::ProcessConnectionState_Reassign(const FMultiServerProxyInt
 	ReassignPlayerController(Route, GameServerConnection, PlayerController);
 }
 
+void UProxyNetDriver::UpdateProxyNetworkMetrics()
+{
+	extern ENGINE_API float GAverageFPS;
+	extern ENGINE_API float GAverageMS;
+
+	GetMetrics()->SetFloat(UE::MultiServerProxy::Metrics::FrameTimeMS, GAverageMS);
+	GetMetrics()->SetFloat(UE::MultiServerProxy::Metrics::FPS, GAverageFPS);
+	GetMetrics()->SetInt(UE::MultiServerProxy::Metrics::NumProxyRoutes, InternalRoutes.Num());
+
+	AggregateBackendNetDriverMetrics();
+}
+
+void UProxyNetDriver::AggregateBackendNetDriverMetrics()
+{
+	// Add together any metrics from the backend net drivers and store them in a metric in the frontend
+	// proxy net driver which will be reported to metrics database listeners.
+
+	static const TMap<FName, FName> BackendMetricToFrontendMetricInt = {
+		{	UE::Net::Metric::InRate, 	UE::MultiServerProxy::Metrics::AggregatedBackendInRate	},
+		{	UE::Net::Metric::OutRate, 	UE::MultiServerProxy::Metrics::AggregatedBackendOutRate	}
+	};
+
+	for (const TPair<FName, FName>& Pair : BackendMetricToFrontendMetricInt)
+	{
+		GetMetrics()->SetInt(Pair.Value, 0);
+	}
+	
+	for (const FGameServerConnectionState& Connection : GameServerConnections)
+	{
+		if (!Connection.NetDriver)
+		{
+			continue;
+		}
+
+		UNetworkMetricsDatabase* BackendMetrics = Connection.NetDriver->GetMetrics();
+		if (!BackendMetrics)
+		{
+			continue;
+		}
+
+		for (const TPair<FName, FName>& Pair : BackendMetricToFrontendMetricInt)
+		{
+			const FName BackendMetricName = Pair.Key;
+			const FName FrontendMetricName = Pair.Value;
+
+			ensure(BackendMetrics->Contains(BackendMetricName));
+			ensure(GetMetrics()->Contains(FrontendMetricName));
+
+			int64 BackendMetricValue = BackendMetrics->GetInt(BackendMetricName);
+			GetMetrics()->IncrementInt(FrontendMetricName, BackendMetricValue);
+		}
+	}
+}
+
+void UProxyNetDriver::ResetProxyNetworkMetrics()
+{
+	// Reset any metrics that are incremented while ticking the net driver and not just set from existing values.
+
+	GetMetrics()->SetInt(UE::MultiServerProxy::Metrics::NumForwardClientRPC, 0);
+	GetMetrics()->SetInt(UE::MultiServerProxy::Metrics::NumForwardBackendRPC, 0);
+	GetMetrics()->SetInt(UE::MultiServerProxy::Metrics::NumReassignPlayerController, 0);
+}
+
 void UProxyNetDriver::FlushSplitJoinRequests(FGameServerConnectionState* GameServerConnectionState)
 {
 	EConnectionState ParentConnectionState = GameServerConnectionState->NetDriver->ServerConnection->GetConnectionState();
diff --git a/Engine/Plugins/Runtime/MultiServerReplication/Source/MultiServerReplication/Public/MultiServerProxy.h b/Engine/Plugins/Runtime/MultiServerReplication/Source/MultiServerReplication/Public/MultiServerProxy.h
--- a/Engine/Plugins/Runtime/MultiServerReplication/Source/MultiServerReplication/Public/MultiServerProxy.h
+++ b/Engine/Plugins/Runtime/MultiServerReplication/Source/MultiServerReplication/Public/MultiServerProxy.h
@@ -265,6 +265,9 @@ public:
 	virtual void InternalProcessRemoteFunction(class AActor* Actor, class UObject* SubObject, class UNetConnection* Connection, class UFunction* Function, void* Parms, FOutParmRec* OutParms, FFrame* Stack, bool bIsServer) override;
 	virtual bool ShouldClientDestroyActor(AActor* Actor, EChannelCloseReason CloseReason) const;
 
+	virtual bool ShouldUpdateStats() const override;
+	virtual bool ShouldRegisterMetricsDatabaseListeners() const override;
+
 private:
 
 	UPROPERTY()
@@ -317,6 +320,7 @@ public:
 
 	/** UNetDriver interface function. */
 	virtual bool InitBase(bool bInitAsClient, FNetworkNotify* InNotify, const FURL& URL, bool bReuseAddressAndPort, FString& Error) override;
+	virtual void SetupNetworkMetrics() override;
 	virtual bool InitConnect(class FNetworkNotify* InNotify, const FURL& ConnectURL, FString& Error) override;
 	virtual void ForwardRemoteFunction(UObject* RootObject, UObject* SubObject, UFunction* Function, void* Parms) override;
 	virtual bool ShouldReplicateFunction(AActor* Actor, UFunction* Function) const override;
@@ -328,6 +332,7 @@ public:
 	virtual int32 ServerReplicateActors(float DeltaSeconds) override;
 	virtual bool CanDowngradeActorRole(UNetConnection* ProxyConnection, AActor* Actor) const override;
 	virtual void Shutdown() override;
+	virtual void TickFlush(float DeltaSeconds) override;
 	
 	/** Get the next identifier for outgoing connections to game servers. */
 	int32 GetNextGameServerClientId();
@@ -363,6 +368,10 @@ private:
 	void ProcessConnectionState_ConnectedPrimary(const FMultiServerProxyInternalConnectionRoute* Route, UNetConnection* GameServerConnection, APlayerController* PlayerController);
 	void ProcessConnectionState_Reassign(const FMultiServerProxyInternalConnectionRoute* Route, UNetConnection* GameServerConnection, APlayerController* PlayerController);
 
+	void UpdateProxyNetworkMetrics();
+	void AggregateBackendNetDriverMetrics();
+	void ResetProxyNetworkMetrics();
+
 	/** Set all of the configuration options that disable executing actor functionality or game specific code. */
 	void DisableActorLogicAndGameCode();
 	
diff --git a/Engine/Source/Runtime/Engine/Classes/Engine/NetDriver.h b/Engine/Source/Runtime/Engine/Classes/Engine/NetDriver.h
--- a/Engine/Source/Runtime/Engine/Classes/Engine/NetDriver.h
+++ b/Engine/Source/Runtime/Engine/Classes/Engine/NetDriver.h
@@ -1043,6 +1043,15 @@ public:
 
 	/** Creates a channel of each type that is set as bIniitalServer for the given connection. */
 	ENGINE_API void CreateInitialServerChannels(UNetConnection* ClientConnection);
+	
+protected:
+	/** 
+	 * Register each metric used by the networking system. 
+	 *
+	 * If a sub-class needs to create new metrics, those metrics should be added to the metrics database in the sub-class
+	 * implementation of this function. Metrics should only be added to the database after calling Super::SetupNetworkMetrics().
+	 */
+	ENGINE_API virtual void SetupNetworkMetrics();
 
 private:
 
@@ -1064,9 +1073,6 @@ private:
 	UPROPERTY()
 	TMap<FName, TObjectPtr<UNetworkMetricsBaseListener>> NetworkMetricsListeners;
 
-	/** Register each metric used by the networking system. */
-	void SetupNetworkMetrics();
-
 	/** Register metric listeners provided by the *.ini file. */
 	void SetupNetworkMetricsListeners();
 
@@ -1983,6 +1989,14 @@ public:
 	UE::Net::FNetTokenStore* GetNetTokenStore() { return NetTokenStore.Get(); }
 	
 	DECLARE_MULTICAST_DELEGATE_OneParam(FOnNetTokenStoreReady, UNetDriver*);
+
+protected:
+	/** Return true if this driver should update network stats in UpdateNetworkStats(). */
+	ENGINE_API virtual bool ShouldUpdateStats() const;
+
+	/** Return true if this driver should register metrics database listeners based on UNetworkMetricsConfig ini configuration. */
+	ENGINE_API virtual bool ShouldRegisterMetricsDatabaseListeners() const;
+
 private:
 	FOnNetTokenStoreReady OnNetTokenStoreReadyDelegate;
 public:
diff --git a/Engine/Source/Runtime/Engine/Private/Net/NetworkMetricsDatabase.cpp b/Engine/Source/Runtime/Engine/Private/Net/NetworkMetricsDatabase.cpp
--- a/Engine/Source/Runtime/Engine/Private/Net/NetworkMetricsDatabase.cpp
+++ b/Engine/Source/Runtime/Engine/Private/Net/NetworkMetricsDatabase.cpp
@@ -82,6 +82,20 @@ bool UNetworkMetricsDatabase::SetMaxFloat(const FName MetricName, float Value)
 	}
 }
 
+float UNetworkMetricsDatabase::GetFloat(const FName MetricName)
+{
+	UE::Net::FNetworkMetric<float>* Metric = MetricFloats.Find(MetricName);
+	if (Metric)
+	{
+		return Metric->Value;
+	}
+	else
+	{
+		UE_LOG(LogNetworkMetrics, Warning, TEXT("UNetworkMetricsDatabase::GetFloat Cannot find float metric %s."), *MetricName.ToString());
+		return 0.0f;
+	}
+}
+
 bool UNetworkMetricsDatabase::SetInt(const FName MetricName, int64 Value)
 {
 	UE::Net::FNetworkMetric<int64>* Metric = MetricInts.Find(MetricName);
@@ -142,6 +156,20 @@ bool UNetworkMetricsDatabase::IncrementInt(const FName MetricName, int64 Value)
 	}
 }
 
+int64 UNetworkMetricsDatabase::GetInt(const FName MetricName)
+{
+	UE::Net::FNetworkMetric<int64>* Metric = MetricInts.Find(MetricName);
+	if (Metric)
+	{
+		return Metric->Value;
+	}
+	else
+	{
+		UE_LOG(LogNetworkMetrics, Warning, TEXT("UNetworkMetricsDatabase::GetInt Cannot find integer metric %s."), *MetricName.ToString());
+		return 0;
+	}
+}
+
 bool UNetworkMetricsDatabase::Contains(const FName MetricName) const
 {
 	return MetricTypes.Contains(MetricName);
@@ -202,6 +230,11 @@ void UNetworkMetricsDatabase::Reset()
 	MetricInts.Reset();
 	MetricFloats.Reset();
 	MetricTypes.Reset();
+	ResetListeners();
+}
+
+void UNetworkMetricsDatabase::ResetListeners()
+{
 	LastReportListener.Reset();
 	ListenersToMetrics.Reset();
 }
@@ -246,6 +279,23 @@ UNetworkMetricsBaseListener::UNetworkMetricsBaseListener() :
 
 }
 
+void UNetworkMetricsLog::Report(const UE::Net::FNetworkMetricSnapshot& Snapshot)
+{
+	FString MetricsStr;
+
+	for (const UE::Net::FNetworkMetric<int64>& Metric : Snapshot.MetricInts)
+	{
+		MetricsStr += FString::Printf(TEXT("%s=%d "), *Metric.Name.ToString(), Metric.Value);
+	}
+
+	for (const UE::Net::FNetworkMetric<float>& Metric : Snapshot.MetricFloats)
+	{
+		MetricsStr += FString::Printf(TEXT("%s=%2.2f "), *Metric.Name.ToString(), Metric.Value);
+	}
+
+	UE_LOG(LogNetworkMetrics, Log, TEXT("Network Metrics: %s"), *MetricsStr);
+}
+
 UNetworkMetricsCSV::UNetworkMetricsCSV() :
 	CategoryIndex(-1)
 {
diff --git a/Engine/Source/Runtime/Engine/Private/Net/NetworkMetricsDefs.cpp b/Engine/Source/Runtime/Engine/Private/Net/NetworkMetricsDefs.cpp
--- a/Engine/Source/Runtime/Engine/Private/Net/NetworkMetricsDefs.cpp
+++ b/Engine/Source/Runtime/Engine/Private/Net/NetworkMetricsDefs.cpp
@@ -74,6 +74,7 @@ const FName NetGUIDOutRate("NetGUIDOutRate");
 const FName NetNumClients("NetNumClients");
 const FName NumClients("NumClients");
 const FName NumConnections("NumConnections");
+const FName NumAllConnections("NumAllConnections");
 const FName NumConsideredActors("NumConsideredActors");
 const FName NumInitiallyDormantActors("NumInitiallyDormantActors");
 const FName PrioritizedActors("PrioritizedActors");
diff --git a/Engine/Source/Runtime/Engine/Private/NetDriver.cpp b/Engine/Source/Runtime/Engine/Private/NetDriver.cpp
--- a/Engine/Source/Runtime/Engine/Private/NetDriver.cpp
+++ b/Engine/Source/Runtime/Engine/Private/NetDriver.cpp
@@ -1715,7 +1715,7 @@ bool UNetDriver::InitBase(bool bInitAsClient, FNetworkNotify* InNotify, const FU
 	// Add all of the metrics used by the networking system and register metrics listeners.
 	SetupNetworkMetrics();
 
-	if (NetDriverDefinition == NAME_GameNetDriver)
+	if (ShouldRegisterMetricsDatabaseListeners())
 	{
 		SetupNetworkMetricsListeners();
 	}
@@ -1899,9 +1899,12 @@ void UNetDriver::SetupNetworkMetrics()
 	// (Server only) The number of active client connections (can be smaller than NumConnections).
 	GetMetrics()->CreateInt(UE::Net::Metric::NumClients, 0);
 
-	// The number of connections.
+	// The number of parent connections.
 	GetMetrics()->CreateInt(UE::Net::Metric::NumConnections, 0);
 
+	// The number of connections including any child connections.
+	GetMetrics()->CreateInt(UE::Net::Metric::NumAllConnections, 0);
+
 	// (Server only) The number of actors considered for replication across all connections.
 	GetMetrics()->CreateInt(UE::Net::Metric::NumConsideredActors, 0);
 
@@ -1976,15 +1979,15 @@ void UNetDriver::SetupNetworkMetricsListeners()
 				continue;
 			}
 
-			if (!GetMetrics()->Contains(MetricConfig.MetricName))
-			{
-				UE_LOG(LogNet, Warning, TEXT("NetDriver::SetupNetworkMetricListeners: Cannot register metric listener %s to an unknown metric '%s'."), *MetricConfig.Class.ToString(), *MetricConfig.MetricName.ToString());
-				continue;
-			}
-
 			Listener = &NetworkMetricsListeners.Add(ListenerClassName, NewObject<UNetworkMetricsBaseListener>(this, ListenerClass));
 		}
 
+		if (!GetMetrics()->Contains(MetricConfig.MetricName))
+		{
+			UE_LOG(LogNet, Warning, TEXT("NetDriver::SetupNetworkMetricListeners: Cannot register metric listener %s to an unknown metric '%s'."), *MetricConfig.Class.ToString(), *MetricConfig.MetricName.ToString());
+			continue;
+		}
+
 		GetMetrics()->Register(MetricConfig.MetricName, Listener->Get());
 		UE_LOG(LogNet, Verbose, TEXT("Registering network metrics listener %s for metric %s"), *ListenerClassName.ToString(), *MetricConfig.MetricName.ToString());
 	}
@@ -7697,6 +7700,16 @@ void UNetDriver::OnLevelRemovedFromWorld(class ULevel* Level, class UWorld* InWo
 	}
 }
 
+bool UNetDriver::ShouldUpdateStats() const
+{
+	return NetDriverName == NAME_GameNetDriver;
+}
+
+bool UNetDriver::ShouldRegisterMetricsDatabaseListeners() const
+{
+	return NetDriverDefinition == NAME_GameNetDriver;
+}
+
 bool UNetDriver::ShouldSkipRepNotifies() const
 {
 #if !UE_BUILD_SHIPPING
@@ -8251,7 +8264,7 @@ void UNetDriver::UpdateNetworkStats()
 
 		const double CurrentRealtimeSeconds = FPlatformTime::Seconds();
 		// Update network stats (only main game net driver for now) if stats or perf counters are used
-		if (NetDriverName == NAME_GameNetDriver &&
+		if (ShouldUpdateStats() &&
 			CurrentRealtimeSeconds - StatUpdateTime > StatPeriod)
 		{
 			int32 ClientInBytesMax = 0;
@@ -8524,6 +8537,8 @@ void UNetDriver::UpdateNetworkStats()
 				const int kNumBuckets = 8;	// evenly spaced with increment of 30 ms; last bucket collects all off-scale pings as well
 				if (ClientConnections.Num() > 0)
 				{
+					int32 NumAllConnections = 0;
+
 					// Update per connection statistics
 					float MinPing = UE_MAX_FLT;
 					float AvgPing = 0;
@@ -8536,6 +8551,9 @@ void UNetDriver::UpdateNetworkStats()
 					{
 						UNetConnection* Connection = ClientConnections[i];
 
+						NumAllConnections++;
+						NumAllConnections += Connection->Children.Num();
+
 						if (Connection != nullptr)
 						{
 							if (Connection->PlayerController && Connection->PlayerController->PlayerState)
@@ -8579,6 +8597,8 @@ void UNetDriver::UpdateNetworkStats()
 						AvgPing /= static_cast<float>(PingCount);
 					}
 
+					GetMetrics()->SetInt(UE::Net::Metric::NumAllConnections, NumAllConnections);
+
 					GetMetrics()->SetFloat(UE::Net::Metric::AvgPing, AvgPing);
 					GetMetrics()->SetFloat(UE::Net::Metric::MaxPing, MaxPing);
 					GetMetrics()->SetFloat(UE::Net::Metric::MinPing, MinPing);
@@ -8599,6 +8619,8 @@ void UNetDriver::UpdateNetworkStats()
 				}
 				else
 				{
+					GetMetrics()->SetInt(UE::Net::Metric::NumAllConnections, 0);
+
 					GetMetrics()->SetFloat(UE::Net::Metric::AvgPing, 0.0f);
 					GetMetrics()->SetFloat(UE::Net::Metric::MaxPing, -FLT_MAX);
 					GetMetrics()->SetFloat(UE::Net::Metric::MinPing, FLT_MAX);
diff --git a/Engine/Source/Runtime/Engine/Public/Net/NetworkMetricsDatabase.h b/Engine/Source/Runtime/Engine/Public/Net/NetworkMetricsDatabase.h
--- a/Engine/Source/Runtime/Engine/Public/Net/NetworkMetricsDatabase.h
+++ b/Engine/Source/Runtime/Engine/Public/Net/NetworkMetricsDatabase.h
@@ -47,6 +47,8 @@ public:
 	UE_API bool SetMinFloat(const FName MetricName, float Value);
 	/* Set the value of a floating point metric if it's bigger than the existing value. */
 	UE_API bool SetMaxFloat(const FName MetricName, float Value);
+	/** Get the value of an existing floating point metric. */
+	UE_API float GetFloat(const FName MetricName);
 	/* Set the value of an existing integer metric. */
 	UE_API bool SetInt(const FName MetricName, int64 Value);
 	/* Set the value of an integer metric if it's smaller than the existing value. */
@@ -55,6 +57,8 @@ public:
 	UE_API bool SetMaxInt(const FName MetricName, int64 Value);
 	/* Increment the value of an existing integer metric. */
 	UE_API bool IncrementInt(const FName MetricName, int64 Value);
+	/** Get the value of an existing integer metric. */
+	UE_API int64 GetInt(const FName MetricName);
 	/* Returns true if a metric has been created in the database. */
 	UE_API bool Contains(const FName MetricName) const;
 
@@ -62,6 +66,8 @@ public:
 	UE_API void ProcessListeners();
 	/* Remove all registered metrics and listeners. */
 	UE_API void Reset();
+	/* Remove all registered listeners. */
+	UE_API void ResetListeners();
 
 	/* Register a listener to be called for a given metric. */
 	UE_API void Register(const FName MetricName, TWeakObjectPtr<UNetworkMetricsBaseListener> Reporter);
@@ -165,6 +171,30 @@ private:
 	double IntervalSeconds;
 };
 
+/** 
+ * A metrics listener that prints stats to a log statement at a fixed frequency to the LogNetworkMetrics log category.
+ *
+ * To use UNetworkMetricsLog in a configuration file, the class must be used when registering a listener:
+ * 
+ * [/Script/Engine.NetworkMetricsConfig]
+ * +Listeners=(MetricName=ExampleMetric, ClassName=/Script/Engine.NetworkMetricsLog)
+ *
+ * Additionally, the following log category must be set:
+ *
+ * [Core.Log]
+ * LogNetworkMetrics=Log
+ */
+UCLASS(MinimalAPI)
+class UNetworkMetricsLog : public UNetworkMetricsBaseListener
+{
+	GENERATED_BODY()
+
+public:
+	virtual ~UNetworkMetricsLog() = default;
+
+	UE_API virtual void Report(const UE::Net::FNetworkMetricSnapshot& Snapshot);
+};
+
 /**
  * A metrics listener that reports an array of metrics to CSV.
  * 
diff --git a/Engine/Source/Runtime/Engine/Public/Net/NetworkMetricsDefs.h b/Engine/Source/Runtime/Engine/Public/Net/NetworkMetricsDefs.h
--- a/Engine/Source/Runtime/Engine/Public/Net/NetworkMetricsDefs.h
+++ b/Engine/Source/Runtime/Engine/Public/Net/NetworkMetricsDefs.h
@@ -76,6 +76,7 @@ extern ENGINE_API const FName NetGUIDOutRate;
 extern ENGINE_API const FName NetNumClients;
 extern ENGINE_API const FName NumClients;
 extern ENGINE_API const FName NumConnections;
+extern ENGINE_API const FName NumAllConnections;
 extern ENGINE_API const FName NumConsideredActors;
 extern ENGINE_API const FName NumInitiallyDormantActors;
 extern ENGINE_API const FName PrioritizedActors;
