Remove player id logging from the multiserver proxy to fix missing symbol build error.

diff --git a/Engine/Plugins/Runtime/MultiServerReplication/Source/MultiServerReplication/Private/MultiServerProxy.cpp b/Engine/Plugins/Runtime/MultiServerReplication/Source/MultiServerReplication/Private/MultiServerProxy.cpp
--- a/Engine/Plugins/Runtime/MultiServerReplication/Source/MultiServerReplication/Private/MultiServerProxy.cpp
+++ b/Engine/Plugins/Runtime/MultiServerReplication/Source/MultiServerReplication/Private/MultiServerProxy.cpp
@@ -916,11 +916,10 @@ void UProxyNetDriver::DisableActorLogicAndGameCode()
 
 void UProxyNetDriver::GameServerAssignPlayerController(UNetConnection* ChildGameServerConnection, UNetConnection* NotUsedParentGameServerConnection, APlayerController* GameServerPlayerController)
 {
-	UE_LOG(LogNetProxy, Log, TEXT("[%d] Received a new player controller %s:%s for connection %s:%s with player id %s."), 
+	UE_LOG(LogNetProxy, Log, TEXT("[%d] Received a new player controller %s:%s for connection %s:%s."), 
 		   GameServerPlayerController->GetClientHandshakeId(),
 		   *GameServerPlayerController->GetName(), *GameServerPlayerController->GetClass()->GetName(),
-		   *ChildGameServerConnection->GetDriver()->GetName(), *ChildGameServerConnection->GetName(),
-		   *ChildGameServerConnection->PlayerId.ToDebugString());
+		   *ChildGameServerConnection->GetDriver()->GetName(), *ChildGameServerConnection->GetName());
 
 	const uint32 ClientHandshakeId = GameServerPlayerController->GetClientHandshakeId();
 	FMultiServerProxyInternalConnectionRoute* Route = InternalRoutes.Find(ClientHandshakeId);
