Log the player id on the multiserver proxy when assinging a new player controller.

diff --git a/Engine/Plugins/Runtime/MultiServerReplication/Source/MultiServerConfiguration/MultiServerConfiguration.Build.cs b/Engine/Plugins/Runtime/MultiServerReplication/Source/MultiServerConfiguration/MultiServerConfiguration.Build.cs
--- a/Engine/Plugins/Runtime/MultiServerReplication/Source/MultiServerConfiguration/MultiServerConfiguration.Build.cs
+++ b/Engine/Plugins/Runtime/MultiServerReplication/Source/MultiServerConfiguration/MultiServerConfiguration.Build.cs
@@ -9,7 +9,8 @@ public class MultiServerConfiguration : ModuleRules
         PrivateDependencyModuleNames.AddRange(
 			new string[] {
                 "Core",
-                "CoreUObject"
+                "CoreUObject",
+                "CoreOnline"
             }
         );
     }
diff --git a/Engine/Plugins/Runtime/MultiServerReplication/Source/MultiServerReplication/Private/MultiServerProxy.cpp b/Engine/Plugins/Runtime/MultiServerReplication/Source/MultiServerReplication/Private/MultiServerProxy.cpp
--- a/Engine/Plugins/Runtime/MultiServerReplication/Source/MultiServerReplication/Private/MultiServerProxy.cpp
+++ b/Engine/Plugins/Runtime/MultiServerReplication/Source/MultiServerReplication/Private/MultiServerProxy.cpp
@@ -916,10 +916,11 @@ void UProxyNetDriver::DisableActorLogicAndGameCode()
 
 void UProxyNetDriver::GameServerAssignPlayerController(UNetConnection* ChildGameServerConnection, UNetConnection* NotUsedParentGameServerConnection, APlayerController* GameServerPlayerController)
 {
-	UE_LOG(LogNetProxy, Log, TEXT("[%d] Received a new player controller %s:%s for connection %s:%s."), 
+	UE_LOG(LogNetProxy, Log, TEXT("[%d] Received a new player controller %s:%s for connection %s:%s for player %s."), 
 		   GameServerPlayerController->GetClientHandshakeId(),
 		   *GameServerPlayerController->GetName(), *GameServerPlayerController->GetClass()->GetName(),
-		   *ChildGameServerConnection->GetDriver()->GetName(), *ChildGameServerConnection->GetName());
+		   *ChildGameServerConnection->GetDriver()->GetName(), *ChildGameServerConnection->GetName(),
+		   *ChildGameServerConnection->PlayerId.ToDebugString());
 
 	const uint32 ClientHandshakeId = GameServerPlayerController->GetClientHandshakeId();
 	FMultiServerProxyInternalConnectionRoute* Route = InternalRoutes.Find(ClientHandshakeId);
