Move the CoreOnline dependency to the correct multi-server build file.

diff --git a/Engine/Plugins/Runtime/MultiServerReplication/Source/MultiServerConfiguration/MultiServerConfiguration.Build.cs b/Engine/Plugins/Runtime/MultiServerReplication/Source/MultiServerConfiguration/MultiServerConfiguration.Build.cs
--- a/Engine/Plugins/Runtime/MultiServerReplication/Source/MultiServerConfiguration/MultiServerConfiguration.Build.cs
+++ b/Engine/Plugins/Runtime/MultiServerReplication/Source/MultiServerConfiguration/MultiServerConfiguration.Build.cs
@@ -9,8 +9,7 @@ public class MultiServerConfiguration : ModuleRules
         PrivateDependencyModuleNames.AddRange(
 			new string[] {
                 "Core",
-                "CoreUObject",
-                "CoreOnline"
+                "CoreUObject"
             }
         );
     }
diff --git a/Engine/Plugins/Runtime/MultiServerReplication/Source/MultiServerReplication/MultiServerReplication.Build.cs b/Engine/Plugins/Runtime/MultiServerReplication/Source/MultiServerReplication/MultiServerReplication.Build.cs
--- a/Engine/Plugins/Runtime/MultiServerReplication/Source/MultiServerReplication/MultiServerReplication.Build.cs
+++ b/Engine/Plugins/Runtime/MultiServerReplication/Source/MultiServerReplication/MultiServerReplication.Build.cs
@@ -10,6 +10,7 @@ public class MultiServerReplication : ModuleRules
 			new string[] {
                 "Core",
                 "CoreUObject",
+                "CoreOnline",
                 "Engine",
 				"NetCore",
 				"MultiServerConfiguration"
