Proxy can perform logic that depends on the order in which player controllers are re-assigned on game servers and fix a bug where a character is downgraded to ROLE_SimulatedProxy on the client for one frame during player migration.

diff --git a/Engine/Plugins/Runtime/MultiServerReplication/Source/MultiServerReplication/Private/MultiServerProxy.cpp b/Engine/Plugins/Runtime/MultiServerReplication/Source/MultiServerReplication/Private/MultiServerProxy.cpp
--- a/Engine/Plugins/Runtime/MultiServerReplication/Source/MultiServerReplication/Private/MultiServerProxy.cpp
+++ b/Engine/Plugins/Runtime/MultiServerReplication/Source/MultiServerReplication/Private/MultiServerProxy.cpp
@@ -109,7 +109,7 @@ FNetworkGUID FProxyNetGUIDCache::LookupNetGUIDFromBackendCache(UObject *Object)
 				// Catch cases where the NetGUID is already in the proxy cache.
 				ensure(!GetNetGUID(Object).IsValid());
 
-				UE_LOG(LogNetProxy, Log, TEXT("NetGUID: Registering NetGUID %s for object %s found in the shared NetGUID cache."), *MatchingNetGUID.ToString(), *Object->GetName());
+				UE_LOG(LogNetProxy, Verbose, TEXT("NetGUID: Registering NetGUID %s for object %s found in the shared NetGUID cache."), *MatchingNetGUID.ToString(), *Object->GetName());
 				RegisterNetGUID_Server(MatchingNetGUID, Object);
 				return MatchingNetGUID;
 			}
@@ -1006,40 +1006,29 @@ void UProxyNetDriver::ReassignPlayerController(const FMultiServerProxyInternalCo
 	PlayerController->NetConnection = GameServerConnection;
 	PlayerController->Player = Route->Player;
 	Route->PlayerController = PlayerController;
+	
+	UE_LOG(LogNetProxy, Log, TEXT("[%u] Replacing player controller %s with %s on route with game server connection %s and player id %s"),
+		   Route->ClientHandshakeId,
+		   *UE::MultiServerProxy::Private::PlayerControllerToString(PrevPlayerController),
+		   *UE::MultiServerProxy::Private::PlayerControllerToString(Route->PlayerController),
+		   *UE::MultiServerProxy::Private::ConnectionToString(Route->GameServerConnection),
+		   *GetNameSafe(Route->Player));
 
 	// The proxy listens for 2 events when a player controller migrates from one game server to another:
 	//	* The player controller on the connection to game server A will change to an instance of ANoPawnPlayerController.
 	//	* The player controller on the connection to game server B will change to the instance of the game's player controller.
 	//
-	// IMPORTANT: The code below must not assume that these events will arrive in a specific order, or even arrive at all.
+	// IMPORTANT: 
+	//		The code below must not assume that these events will arrive in a specific order, or even arrive at all.
 
 	if (PlayerController->IsA(ANoPawnPlayerController::StaticClass()))
 	{
-		UE_LOG(LogNetProxy, Log, TEXT("[%u] Deregistering player %s as a primary game server."), 
-			   Route->ClientHandshakeId,
-			   *Route->Player->GetName());
-
-		PlayerToPrimaryGameServerPlayerController.Remove(Route->Player);
-		Route->State = EProxyConnectionState::Connected;
+		ReceivedReassignedNoPawnPlayerController(Route);
 	}
 	else
 	{
-		UE_LOG(LogNetProxy, Log, TEXT("[%u] Registering player %s as a primary game server using player controller %s."), 
-			   Route->ClientHandshakeId,
-			   *Route->Player->GetName(),
-			   *UE::MultiServerProxy::Private::PlayerControllerToString(Route->PlayerController));
-
-		PlayerToPrimaryGameServerPlayerController.Add(Route->Player, Route->PlayerController);
-		ProxyConnectionToPrimaryGameServerPlayerController.Add(Route->ProxyConnection, Route->PlayerController);
-		Route->State = EProxyConnectionState::ConnectedPrimary;
+		ReceivedReassignedGamePlayerController(Route);
 	}
-
-	UE_LOG(LogNetProxy, Log, TEXT("[%u] Replacing player controller %s with %s on route with game server connection %s and player %s"),
-		   Route->ClientHandshakeId,
-		   *UE::MultiServerProxy::Private::PlayerControllerToString(PrevPlayerController),
-		   *UE::MultiServerProxy::Private::PlayerControllerToString(Route->PlayerController),
-		   *UE::MultiServerProxy::Private::ConnectionToString(Route->GameServerConnection),
-		   *GetNameSafe(Route->Player));
 }
 
 void UProxyNetDriver::ProcessConnectionState_Connecting(FMultiServerProxyInternalConnectionRoute* Route, UNetConnection* GameServerConnection, APlayerController* PlayerController)
@@ -1068,11 +1057,6 @@ void UProxyNetDriver::ProcessConnectionState_ConnectingPrimary(FMultiServerProxy
 	Route->ProxyConnection->PlayerController = PlayerController;
 	Route->ProxyConnection->OwningActor = PlayerController;
 	Route->ProxyConnection->LastReceiveTime = GetElapsedTime();
-	
-	UE_LOG(LogNetProxy, Log, TEXT("[%u] Registering player %s as a primary game server using player controller %s."), 
-		   Route->ClientHandshakeId,
-		   *Route->Player->GetName(),
-		   *UE::MultiServerProxy::Private::PlayerControllerToString(Route->PlayerController));
 
 	PlayerToPrimaryGameServerPlayerController.Add(Route->Player, Route->PlayerController);
 	ProxyConnectionToPrimaryGameServerPlayerController.Add(Route->ProxyConnection, Route->PlayerController);
@@ -1165,6 +1149,99 @@ void UProxyNetDriver::ResetProxyNetworkMetrics()
 	GetMetrics()->SetInt(UE::MultiServerProxy::Metrics::NumReassignPlayerController, 0);
 }
 
+void UProxyNetDriver::ReceivedReassignedNoPawnPlayerController(const FMultiServerProxyInternalConnectionRoute* Route)
+{
+	FPlayerControllerReassignment& Reassignment = DeferredReassignment.FindOrAdd(Route->ProxyConnection);
+
+	// If we've already received a no pawn player controller, assume that a new migration has begun and reset
+	// the state of the migration data structure. This can happen if a previous player controller migration didn't complete.
+	if (Reassignment.bReceivedNoPawnPlayerController)
+	{
+		Reassignment = FPlayerControllerReassignment();
+
+		UE_LOG(LogNetProxy, Warning, TEXT("[%u] Resetting existing deferred player controller migration on proxy connection %s."), 
+			   Route->ClientHandshakeId,
+			   *Route->ProxyConnection->GetName());
+	}
+
+	Reassignment.PreviousPrimaryPlayer = Route->Player;
+	Reassignment.PreviousClientHandshakeId = Route->ClientHandshakeId;
+	Reassignment.bReceivedNoPawnPlayerController = true;
+
+	if (Reassignment.bReceivedNoPawnPlayerController && Reassignment.bReceivedGamePlayerController)
+	{
+		FinalizePlayerControllerReassignment(Route->ProxyConnection, Reassignment);
+		DeferredReassignment.Remove(Route->ProxyConnection);
+	}
+	else
+	{
+		Route->State = EProxyConnectionState::PendingReassign;
+	}
+}
+
+void UProxyNetDriver::ReceivedReassignedGamePlayerController(const FMultiServerProxyInternalConnectionRoute* Route)
+{
+	FPlayerControllerReassignment& Reassignment = DeferredReassignment.FindOrAdd(Route->ProxyConnection);
+
+	// If we've already received a new game pawn player controller, assume that a new migration has begun and reset
+	// the state of the migration data structure. This can happen if a previous player controller migration didn't complete.
+	if (Reassignment.bReceivedGamePlayerController)
+	{
+		Reassignment = FPlayerControllerReassignment();
+
+		UE_LOG(LogNetProxy, Warning, TEXT("[%u] Resetting existing deferred player controller reassignment on proxy connection %s."), 
+			   Route->ClientHandshakeId,
+			   *Route->ProxyConnection->GetName());
+	}
+
+	Reassignment.PrimaryPlayer = Route->Player;
+	Reassignment.ClientHandshakeId = Route->ClientHandshakeId;
+	Reassignment.bReceivedGamePlayerController = true;
+
+	PlayerToPrimaryGameServerPlayerController.Add(Reassignment.PrimaryPlayer, Route->PlayerController);
+	ProxyConnectionToPrimaryGameServerPlayerController.Add(Route->ProxyConnection, Route->PlayerController);
+
+	if (Reassignment.bReceivedNoPawnPlayerController && Reassignment.bReceivedGamePlayerController)
+	{
+		FinalizePlayerControllerReassignment(Route->ProxyConnection, Reassignment);
+		DeferredReassignment.Remove(Route->ProxyConnection);
+	}
+	else
+	{
+		Route->State = EProxyConnectionState::PendingReassign;
+	}
+}
+
+void UProxyNetDriver::FinalizePlayerControllerReassignment(UNetConnection* ProxyConnection, const FPlayerControllerReassignment& Reassignment)
+{
+	// Perform any logic that depends on the order that the proxy receives reassigned game and no pawn player controllers, or having received
+	// both controllers.
+
+	UE_LOG(LogNetProxy, Log, TEXT("Finalizing reassignment of primary player from route %d, with player %s, to route %d and player %s."),
+		   Reassignment.PreviousClientHandshakeId,
+		   *Reassignment.PreviousPrimaryPlayer->GetName(),
+		   Reassignment.ClientHandshakeId,
+		   *Reassignment.PrimaryPlayer->GetName());
+	
+	if (Reassignment.PreviousPrimaryPlayer != Reassignment.PrimaryPlayer)
+	{
+		PlayerToPrimaryGameServerPlayerController.Remove(Reassignment.PreviousPrimaryPlayer);
+	}
+
+	const FMultiServerProxyInternalConnectionRoute* PreviousRoute = InternalRoutes.Find(Reassignment.PreviousClientHandshakeId);
+	const FMultiServerProxyInternalConnectionRoute* Route = InternalRoutes.Find(Reassignment.ClientHandshakeId);
+
+	if (ensure(PreviousRoute))
+	{
+		PreviousRoute->State = EProxyConnectionState::Connected;
+	}
+
+	if (ensure(Route))
+	{
+		Route->State = EProxyConnectionState::ConnectedPrimary;
+	}
+}
+
 void UProxyNetDriver::FlushSplitJoinRequests(FGameServerConnectionState* GameServerConnectionState)
 {
 	EConnectionState ParentConnectionState = GameServerConnectionState->NetDriver->ServerConnection->GetConnectionState();
@@ -1212,7 +1289,6 @@ UNetConnection* UProxyNetDriver::GetProxyConnectionFromPrimaryPlayer(ULocalPlaye
 	}
 
 	ensure(Route->Player == Player);
-	ensure(Route->State == EProxyConnectionState::ConnectedPrimary);
 
 	return Route->ProxyConnection;
 }
@@ -1297,8 +1373,6 @@ void UProxyNetDriver::SetRemoteViewTarget(const ANoPawnPlayerController* PlayerC
 		return;
 	}
 
-	ensure(PlayerControllerRoute->State == EProxyConnectionState::Connected);
-
 	if (PlayerControllerRoute->LastViewTargetPos.Equals(ViewTargetPos))
 	{
 		return;
diff --git a/Engine/Plugins/Runtime/MultiServerReplication/Source/MultiServerReplication/Public/MultiServerProxy.h b/Engine/Plugins/Runtime/MultiServerReplication/Source/MultiServerReplication/Public/MultiServerProxy.h
--- a/Engine/Plugins/Runtime/MultiServerReplication/Source/MultiServerReplication/Public/MultiServerProxy.h
+++ b/Engine/Plugins/Runtime/MultiServerReplication/Source/MultiServerReplication/Public/MultiServerProxy.h
@@ -142,6 +142,27 @@ struct FMultiServerProxyInternalConnectionRoute
 	mutable FVector LastViewTargetPos;
 };
 
+/** 
+ * State associated with the reassignment or migration of a player controller from one game server to another.
+*/
+USTRUCT()
+struct FPlayerControllerReassignment
+{
+	GENERATED_BODY()
+
+	bool bReceivedNoPawnPlayerController = false;
+	bool bReceivedGamePlayerController = false;
+
+	uint64 PreviousClientHandshakeId = 0;
+	uint64 ClientHandshakeId = 0;
+
+	UPROPERTY()
+	TObjectPtr<class ULocalPlayer> PreviousPrimaryPlayer;
+
+	UPROPERTY()
+	TObjectPtr<class ULocalPlayer> PrimaryPlayer;
+};
+
 /** 
  * Intercept outgoing connection requests to game servers from the proxy server.
  *
@@ -372,6 +393,23 @@ private:
 	void AggregateBackendNetDriverMetrics();
 	void ResetProxyNetworkMetrics();
 
+	/** 
+	 * Receive a new game player controller for an already open route.
+	 * 
+	 * This route becomes a route to a primary game server.
+	 */
+	void ReceivedReassignedGamePlayerController(const FMultiServerProxyInternalConnectionRoute* Route);
+	
+	/** 
+	 * Receive a new no-pawn player controller for an already open route.
+	 * 
+	 * This route stops being a route to a primary game server.
+	 */
+	void ReceivedReassignedNoPawnPlayerController(const FMultiServerProxyInternalConnectionRoute* Route);
+
+	/** Perform any logic that depends on the order of game and no pawn player controllers being re-assigned. */
+	void FinalizePlayerControllerReassignment(UNetConnection* ProxyConnection, const FPlayerControllerReassignment& Migration);
+
 	/** Set all of the configuration options that disable executing actor functionality or game specific code. */
 	void DisableActorLogicAndGameCode();
 	
@@ -407,6 +445,9 @@ private:
 	UPROPERTY()
 	TMap<TObjectPtr<class ULocalPlayer>, TObjectPtr<const APlayerController>> PlayerToPrimaryGameServerPlayerController;
 
+	UPROPERTY()
+	TMap<TObjectPtr<UNetConnection>, FPlayerControllerReassignment> DeferredReassignment;
+
 	/** Map a proxy connection to the primary game server player controller. */
 	UPROPERTY()
 	TMap<TObjectPtr<UNetConnection>, TObjectPtr<const APlayerController>> ProxyConnectionToPrimaryGameServerPlayerController;
