Small optimizations for the DSTM demo that reduce the number of actors replicated to the client each frame.

diff --git a/Engine/Plugins/Runtime/MultiServerReplication/Source/MultiServerReplication/Private/MultiServerProxy.cpp b/Engine/Plugins/Runtime/MultiServerReplication/Source/MultiServerReplication/Private/MultiServerProxy.cpp
--- a/Engine/Plugins/Runtime/MultiServerReplication/Source/MultiServerReplication/Private/MultiServerProxy.cpp
+++ b/Engine/Plugins/Runtime/MultiServerReplication/Source/MultiServerReplication/Private/MultiServerProxy.cpp
@@ -70,6 +70,9 @@ namespace UE::MultiServerProxy::Metrics
 
 	// The number of player controller's re-assigned on backend game servers.
 	const FName NumReassignPlayerController("NumReassignPlayerController");
+
+	// The number of pawns being replicated to clients.
+	const FName NumReplicatedPawns("NumReplicatedPawns");
 }
 
 FProxyNetGUIDCache::FProxyNetGUIDCache(UNetDriver* NetDriver) : FNetGUIDCache(NetDriver)
@@ -529,7 +532,7 @@ void UProxyBackendNetDriver::InternalProcessRemoteFunction(
 			{
 				bShouldForwardRPC = false;
 
-				UE_LOG(LogNetProxy, Warning, TEXT("RPC: Unable map sub-object %s to actor %s"),
+				UE_LOG(LogNetProxy, Verbose, TEXT("RPC: Unable map sub-object %s to actor %s"),
 					   *SubObjectAsActorComponent->GetName(),
 					   *Actor->GetName());
 			}
@@ -539,7 +542,7 @@ void UProxyBackendNetDriver::InternalProcessRemoteFunction(
 			{
 				bShouldForwardRPC = false;
 
-				UE_LOG(LogNetProxy, Warning, TEXT("RPC: Found an actor %s with more than one component %s."), *Actor->GetName(), *SubObject->GetName());
+				UE_LOG(LogNetProxy, Verbose, TEXT("RPC: Found an actor %s with more than one component %s."), *Actor->GetName(), *SubObject->GetName());
 			}
 			else
 			{
@@ -691,6 +694,7 @@ void UProxyNetDriver::SetupNetworkMetrics()
 	GetMetrics()->CreateInt(UE::MultiServerProxy::Metrics::AggregatedBackendInRate, 0);
 	GetMetrics()->CreateInt(UE::MultiServerProxy::Metrics::AggregatedBackendOutRate, 0);
 	GetMetrics()->CreateInt(UE::MultiServerProxy::Metrics::NumReassignPlayerController, 0);
+	GetMetrics()->CreateInt(UE::MultiServerProxy::Metrics::NumReplicatedPawns, 0);
 }
 
 bool UProxyNetDriver::InitConnect(class FNetworkNotify* InNotify, const FURL& ConnectURL, FString& Error)
@@ -712,14 +716,14 @@ void UProxyNetDriver::ForwardRemoteFunction(UObject* RootObject, UObject* SubObj
 	AActor* OwningActor = Cast<AActor>(RootObject);
 	if (!ensure(OwningActor != nullptr))
 	{
-		UE_LOG(LogNetProxy, Warning, TEXT("RPC: Ignoring %s on %s from proxy connection because it doesn't have an owner."), *Function->GetName(), *RootObject->GetName());
+		UE_LOG(LogNetProxy, Verbose, TEXT("RPC: Ignoring %s on %s from proxy connection because it doesn't have an owner."), *Function->GetName(), *RootObject->GetName());
 		return;
 	}
 
 	ULocalPlayer* OwningPlayer = Cast<ULocalPlayer>(OwningActor->GetNetOwningPlayerAnyRole());
 	if (OwningPlayer == nullptr)
 	{
-		UE_LOG(LogNetProxy, Warning, TEXT("RPC: Ignoring %s on %s from proxy connection because it owner %s doesn't have an owning player"), *Function->GetName(), *RootObject->GetName(), *OwningActor->GetName());
+		UE_LOG(LogNetProxy, Verbose, TEXT("RPC: Ignoring %s on %s from proxy connection because it owner %s doesn't have an owning player"), *Function->GetName(), *RootObject->GetName(), *OwningActor->GetName());
 		return;
 	}
 
@@ -743,7 +747,7 @@ void UProxyNetDriver::ForwardRemoteFunction(UObject* RootObject, UObject* SubObj
 	}
 	else
 	{
-		UE_LOG(LogNetProxy, Warning, TEXT("RPC: Unable to forward %s on %s because player %s isn't mapped to a primary game server."), *Function->GetName(), *RootObject->GetName(), *OwningPlayer->GetName());
+		UE_LOG(LogNetProxy, Verbose, TEXT("RPC: Unable to forward %s on %s because player %s isn't mapped to a primary game server."), *Function->GetName(), *RootObject->GetName(), *OwningPlayer->GetName());
 	}
 }
 
@@ -1147,6 +1151,7 @@ void UProxyNetDriver::ResetProxyNetworkMetrics()
 	GetMetrics()->SetInt(UE::MultiServerProxy::Metrics::NumForwardClientRPC, 0);
 	GetMetrics()->SetInt(UE::MultiServerProxy::Metrics::NumForwardBackendRPC, 0);
 	GetMetrics()->SetInt(UE::MultiServerProxy::Metrics::NumReassignPlayerController, 0);
+	GetMetrics()->SetInt(UE::MultiServerProxy::Metrics::NumReplicatedPawns, 0);
 }
 
 void UProxyNetDriver::ReceivedReassignedNoPawnPlayerController(const FMultiServerProxyInternalConnectionRoute* Route)
@@ -1336,6 +1341,8 @@ void UProxyNetDriver::PrepareStateForRelevancy()
 
 	// Relevancy calcluations are made using the position returned by AActor::GetActorLocation(). Since the proxy is not simulating
 	// player movement the actor's location must be manually set to the location in AActor::ReplicatedMovement.
+	int32 NumReplicatedPawns = 0;
+
 	for (const TSharedPtr<FNetworkObjectInfo>& ObjectInfo : GetNetworkObjectList().GetAllObjects())
 	{
 		APawn* Pawn = Cast<APawn>(ObjectInfo->Actor);
@@ -1345,7 +1352,11 @@ void UProxyNetDriver::PrepareStateForRelevancy()
 		}
 
 		Pawn->SetActorLocation(Pawn->GetReplicatedMovement().Location);
+
+		NumReplicatedPawns++;
 	}
+
+	GetMetrics()->SetInt(UE::MultiServerProxy::Metrics::NumReplicatedPawns, NumReplicatedPawns);
 }
 
 void UProxyNetDriver::UpdateLocalViewTarget(const APlayerController* PrimaryPlayerController)
