Fixed inaccessible constructor from new StructOps code.

diff --git a/Engine/Plugins/Runtime/MultiServerReplication/Source/MultiServerReplication/Public/MultiServerProxy.h b/Engine/Plugins/Runtime/MultiServerReplication/Source/MultiServerReplication/Public/MultiServerProxy.h
--- a/Engine/Plugins/Runtime/MultiServerReplication/Source/MultiServerReplication/Public/MultiServerProxy.h
+++ b/Engine/Plugins/Runtime/MultiServerReplication/Source/MultiServerReplication/Public/MultiServerProxy.h
@@ -93,9 +93,7 @@ struct FGameServerConnectionState
 	UPROPERTY()
 	TArray<FGameServerSplitJoinRequest> PendingSplitJoinRequests;
 
-private:
 	// We are only valid with a URL, but a default constructor is needed for serialization
-	friend class UScriptStruct;
 	FGameServerConnectionState() = default;
 };
 
