Multi-server proxy now supports -ProxyClientPrimaryGameServer=random option to randomize the initial primary game server.

diff --git a/Engine/Plugins/Runtime/MultiServerReplication/Source/MultiServerReplication/Private/MultiServerProxy.cpp b/Engine/Plugins/Runtime/MultiServerReplication/Source/MultiServerReplication/Private/MultiServerProxy.cpp
--- a/Engine/Plugins/Runtime/MultiServerReplication/Source/MultiServerReplication/Private/MultiServerProxy.cpp
+++ b/Engine/Plugins/Runtime/MultiServerReplication/Source/MultiServerReplication/Private/MultiServerProxy.cpp
@@ -218,6 +218,12 @@ void UProxyListenerNotify::NotifyControlMessage(UNetConnection* ProxyConnection,
 	// The NMT_Join message received by a proxy connection should trigger a connection to the backend game servers.
 	if (MessageType == NMT_Join)
 	{
+		if (ProxyNetDriver->bRandomizePrimaryGameServerForNextClient)
+		{
+			ProxyNetDriver->PrimaryGameServerForNextClient = FMath::RandHelper(ProxyNetDriver->GetGameServerConnectionCount());
+			ProxyNetDriver->bRandomizePrimaryGameServerForNextClient = false;
+		}
+
 		// For now, the primary game server is always the first registered game server.
 		for (int32 Index = 0; Index < ProxyNetDriver->GetGameServerConnectionCount(); Index++)
 		{
@@ -711,7 +717,14 @@ bool UProxyNetDriver::InitBase(bool bInitAsClient, FNetworkNotify* InNotify, con
 	FString ClientPrimaryGameServer;
 	if (FParse::Value(FCommandLine::Get(), TEXT("ProxyClientPrimaryGameServer="), ClientPrimaryGameServer))
 	{
-		PrimaryGameServerForNextClient = FCString::Atoi(ToCStr(ClientPrimaryGameServer));
+		if (ClientPrimaryGameServer.Equals(TEXT("random"), ESearchCase::IgnoreCase))
+		{
+			bRandomizePrimaryGameServerForNextClient = true;
+		}
+		else
+		{
+			PrimaryGameServerForNextClient = FCString::Atoi(ToCStr(ClientPrimaryGameServer));
+		}
 	}
 
 	if (FParse::Param(FCommandLine::Get(), TEXT("ProxyCyclePrimaryGameServer")))
diff --git a/Engine/Plugins/Runtime/MultiServerReplication/Source/MultiServerReplication/Public/MultiServerProxy.h b/Engine/Plugins/Runtime/MultiServerReplication/Source/MultiServerReplication/Public/MultiServerProxy.h
--- a/Engine/Plugins/Runtime/MultiServerReplication/Source/MultiServerReplication/Public/MultiServerProxy.h
+++ b/Engine/Plugins/Runtime/MultiServerReplication/Source/MultiServerReplication/Public/MultiServerProxy.h
@@ -470,4 +470,7 @@ private:
 
 	// After a client connects to the proxy increment the value of PrimaryGameServerForNextClient.
 	bool bCyclePrimaryGameServer = false;
+
+	// As a client connects to the proxy, randomize PrimaryGameServerForNextClient to one of the known game servers.
+	bool bRandomizePrimaryGameServerForNextClient = false;
 };
diff --git a/Engine/Source/Programs/AutomationTool/Scripts/LaunchMultiServer.Automation.cs b/Engine/Source/Programs/AutomationTool/Scripts/LaunchMultiServer.Automation.cs
--- a/Engine/Source/Programs/AutomationTool/Scripts/LaunchMultiServer.Automation.cs
+++ b/Engine/Source/Programs/AutomationTool/Scripts/LaunchMultiServer.Automation.cs
@@ -263,8 +263,8 @@ Example running 4 clients locally connecting to single proxy server w/ a non-loc
 		[CommandLine, Help("ProxyBotCount=##", "Specify how many bots to launch PER SERVER in client=proxy mode")]
 		public int ProxyBotCount = 0;
 
-		[CommandLine, Help("ProxyClientPrimary=##", "Specify which game server to prefer as the initial proxy connection")]
-		public int ProxyClientPrimary = -1;
+		[CommandLine, Help("ProxyClientPrimary=##", "Specify which game server to prefer as the initial proxy connection, as an integer or the string \"random\"")]
+		public string ProxyClientPrimary = "";
 
 		[CommandLine, Help("ProxyCyclePrimary", "Specify that each subsequent proxy client should choose a different server")]
 		public bool ProxyCyclePrimary = false;
@@ -537,7 +537,7 @@ Example running 4 clients locally connecting to single proxy server w/ a non-loc
 				CommonProxyServerArgs += $" -ProxyGameServers={string.Join(',', AllGameServerAddresses).TrimEnd(',')}";
 			}
 
-			if (ProxyClientPrimary > 0)
+			if (ProxyClientPrimary.Length > 0)
 			{
 				CommonProxyServerArgs += $" -ProxyClientPrimaryGameServer={ProxyClientPrimary}";
 			}
