Minor changes to the MultiServerNode API

diff --git a/Engine/Plugins/Runtime/MultiServerReplication/Source/MultiServerReplication/Private/MultiServerNode.cpp b/Engine/Plugins/Runtime/MultiServerReplication/Source/MultiServerReplication/Private/MultiServerNode.cpp
--- a/Engine/Plugins/Runtime/MultiServerReplication/Source/MultiServerReplication/Private/MultiServerNode.cpp
+++ b/Engine/Plugins/Runtime/MultiServerReplication/Source/MultiServerReplication/Private/MultiServerNode.cpp
@@ -39,6 +39,27 @@ UMultiServerNode::UMultiServerNode()
 	RetryConnectMaxDelay = 30.0f;
 }
 
+bool UMultiServerNode::AreAllServersConnected() const
+{
+	if (NumExpectedServers <= 1)
+	{
+		return true;
+	}
+
+	// We are doing a bit more logic than "connected" and instead ensuring the Peer ID has been exchanged
+	uint32 NumAcknowledgedPeers = 0;
+	ForEachBeaconClient([&NumAcknowledgedPeers](AMultiServerBeaconClient* BeaconClient)
+	{
+		if (!BeaconClient->GetRemotePeerId().IsEmpty())
+		{
+			NumAcknowledgedPeers++;
+		}
+	});
+
+	const bool bAllPeersAreAcknowledged = NumAcknowledgedPeers >= (NumExpectedServers - 1);
+	return bAllPeersAreAcknowledged;
+}
+
 UMultiServerNode* UMultiServerNode::Create(const FMultiServerNodeCreateParams& Params)
 {
 	UMultiServerNode* NewNode = NewObject<UMultiServerNode>(Params.World);
@@ -46,6 +67,7 @@ UMultiServerNode* UMultiServerNode::Create(const FMultiServerNodeCreateParams& P
 
 	if (bResult)
 	{
+		NewNode->NumExpectedServers = Params.NumServers;
 		NewNode->RegisterTickEvents();
 		return NewNode;
 	}
@@ -284,7 +306,7 @@ AMultiServerBeaconClient* UMultiServerNode::GetBeaconClientForURL(const FString&
 	return nullptr;
 }
 
-void UMultiServerNode::ForEachBeaconClient(TFunctionRef<void(AMultiServerBeaconClient*)> Operation)
+void UMultiServerNode::ForEachBeaconClient(TFunctionRef<void(AMultiServerBeaconClient*)> Operation) const
 {
 	UNetDriver* HostNetDriver = BeaconHost ? BeaconHost->GetNetDriver() : nullptr;
 	if (HostNetDriver)
@@ -308,7 +330,7 @@ void UMultiServerNode::ForEachBeaconClient(TFunctionRef<void(AMultiServerBeaconC
 	}
 }
 
-void UMultiServerNode::ForEachNetDriver(TFunctionRef<void(UNetDriver*)> Operation)
+void UMultiServerNode::ForEachNetDriver(TFunctionRef<void(UNetDriver*)> Operation) const
 {
 	TSet<UNetDriver*> UniqueNetDrivers;
 
@@ -348,4 +370,4 @@ int32 UMultiServerNode::GetConnectionCount() const
 		});
 
 	return ConnectionCount;
-}
\ No newline at end of file
+}
diff --git a/Engine/Plugins/Runtime/MultiServerReplication/Source/MultiServerReplication/Public/MultiServerNode.h b/Engine/Plugins/Runtime/MultiServerReplication/Source/MultiServerReplication/Public/MultiServerNode.h
--- a/Engine/Plugins/Runtime/MultiServerReplication/Source/MultiServerReplication/Public/MultiServerNode.h
+++ b/Engine/Plugins/Runtime/MultiServerReplication/Source/MultiServerReplication/Public/MultiServerNode.h
@@ -72,7 +72,10 @@ public:
 
 	static void ParseCommandLineIntoCreateParams(FMultiServerNodeCreateParams& InOutParams);
 
-	virtual void BeginDestroy();
+	virtual void BeginDestroy() override;
+
+	/** @return true if all of the servers we are expecting have connected and registered with this node */
+	bool AreAllServersConnected() const;
 
 	bool RegisterServer(const FMultiServerNodeCreateParams& Params);
 
@@ -93,9 +96,9 @@ public:
 
 	TSubclassOf<AMultiServerBeaconClient> GetUserBeaconClass() const { return UserBeaconClass; }
 
-	void ForEachBeaconClient(TFunctionRef<void(AMultiServerBeaconClient*)> Operation);
+	void ForEachBeaconClient(TFunctionRef<void(AMultiServerBeaconClient*)> Operation) const;
 
-	void ForEachNetDriver(TFunctionRef<void(UNetDriver*)> Operation);
+	void ForEachNetDriver(TFunctionRef<void(UNetDriver*)> Operation) const;
 
 	int32 GetConnectionCount() const;
 
@@ -123,6 +126,9 @@ private:
 
 	TSubclassOf<AMultiServerBeaconClient> UserBeaconClass;
 
+	// How many servers are we expecting? Used for AreAllServersConnected.
+	uint32 NumExpectedServers = 1;
+
 	// Multi-server nodes handle ticking of their NetDrivers directly instead of letting the world tick them.
 	// This allows us to control the timing of the Tick(Flush|Dispatch) and PostTick(Flush|Dispatch) functions
 	// to ensure they're always called as atomic units. Since MultiServer drivers might be ticked from within
