Add support for closing a child connections in the net driver and closing connections on the multi-server proxy.

diff --git a/Engine/Plugins/Runtime/MultiServerReplication/Source/MultiServerReplication/Private/MultiServerProxy.cpp b/Engine/Plugins/Runtime/MultiServerReplication/Source/MultiServerReplication/Private/MultiServerProxy.cpp
--- a/Engine/Plugins/Runtime/MultiServerReplication/Source/MultiServerReplication/Private/MultiServerProxy.cpp
+++ b/Engine/Plugins/Runtime/MultiServerReplication/Source/MultiServerReplication/Private/MultiServerProxy.cpp
@@ -45,6 +45,12 @@ namespace UE::MultiServerProxy::Private
 		TEXT("net.proxy.NonPrimarySetViewTargetInterval"),
 		NonPrimarySetViewTargetIntervalSec,
 		TEXT("The interval, in seconds, between updating the view target position on non-primary game servers."));
+	
+	bool bEnableDisconnectionSupport = true;
+	static FAutoConsoleVariableRef CVarEnableDisconnectionSupport(
+		TEXT("net.proxy.EnableDisconnectionSupport"),
+		bEnableDisconnectionSupport,
+		TEXT("When true, when a client disconnects from the proxy it will close the connection to the backend server."));
 }
 
 namespace UE::MultiServerProxy::Metrics
@@ -71,9 +77,6 @@ namespace UE::MultiServerProxy::Metrics
 
 	// The number of player controller's re-assigned on backend game servers.
 	const FName NumReassignPlayerController("NumReassignPlayerController");
-
-	// The number of pawns being replicated to clients.
-	const FName NumReplicatedPawns("NumReplicatedPawns");
 }
 
 FProxyNetGUIDCache::FProxyNetGUIDCache(UNetDriver* NetDriver) : FNetGUIDCache(NetDriver)
@@ -166,7 +169,31 @@ bool UGameServerNotify::NotifyAcceptingChannel(class UChannel* Channel)
 
 void UGameServerNotify::NotifyControlMessage(UNetConnection* GameServerConnection, uint8 MessageType, class FInBunch& Bunch)
 {
-	Super::NotifyControlMessage(GameServerConnection, MessageType, Bunch);
+	// NMT_CloseChildConnection is handled by UWorld::NotifyControlMessage which is bypassed by UGameServerNotify because it's
+	// a sub-class of UPendingNetGame. For this reason these messages have to be explicitly handled below otherwise child connections 
+	// to the game server will not close correctly.
+	if (MessageType == NMT_CloseChildConnection)
+	{
+		int32 NetPlayerIndexToClose;
+		if (FNetControlMessage<NMT_CloseChildConnection>::Receive(Bunch, NetPlayerIndexToClose))
+		{
+			UE_LOG(LogNet, Log, TEXT("Received NMT_CloseChildConnection with NetPlayerIndex %d"), NetPlayerIndexToClose);
+
+			if (UChildConnection* ChildConnection = GameServerConnection->GetChildConnectionWithNetPlayerIndex(NetPlayerIndexToClose))
+			{
+				// The client shouldn't destroy actors as it will be done by the server.
+				ChildConnection->CloseAndRemoveChild(UE::Net::ECloseChildFlags::None);
+			}
+			else
+			{
+				UE_LOG(LogNet, Warning, TEXT("Unable to find a child connection matching NetPlayerIndex %d"), NetPlayerIndexToClose);
+			}
+		}
+	}
+	else
+	{
+		Super::NotifyControlMessage(GameServerConnection, MessageType, Bunch);
+	}
 
 	if (MessageType == NMT_Welcome)
 	{
@@ -275,25 +302,23 @@ void UProxyListenerNotify::ConnectToGameServer(UNetConnection* ProxyConnection,
 		UE_LOG(LogNetProxy, Log, TEXT("Created a Proxy Game Server NetDriver (name=%s, url=%s)"), *GameServerConnectionState->NetDriver->GetName(), *GameServerConnectionState->GameServerURL.ToString());
 	}
 
-	// Add a player to use on the game server.
-	const FPlatformUserId GameServerClientId = FPlatformUserId::CreateFromInternalId(ProxyNetDriver->GetNextGameServerClientId());
-	ULocalPlayer* NewPlayer = NewObject<ULocalPlayer>(GEngine, GEngine->LocalPlayerClass);
-	ProxyNetDriver->GetWorld()->GetGameInstance()->AddLocalPlayer(NewPlayer, GameServerClientId);
-
-	// The new player will use the same unique identifier as the incoming proxy connection so that it will be propogated up to 
-	// the game servers through UNetConnection::PlayerId. This way each game server's incoming connection will have a PlayerId 
-	// that corresponds to a client connected to the proxy.
-	NewPlayer->SetCachedUniqueNetId(ProxyConnection->PlayerId);
-
 	const uint32 ClientHandshakeId = ProxyNetDriver->GetNextClientHandshakeId();
 
 	UNetConnection* GameServerConnection = GameServerConnectionState->NetDriver->ServerConnection;
 	const bool bIsFirstGameServerConnection = (GameServerConnection == nullptr);
 	if (bIsFirstGameServerConnection)
 	{
+		// The parent connection will not be associated with a player on the proxy and only exists so to use the control channel
+		// and to allow all players (child connections) to be closed without having to worry about closing a parent connection and
+		// promoting a child connection to parent connection.
+
+		const FPlatformUserId DummyClientId = FPlatformUserId::CreateFromInternalId(ProxyNetDriver->GetNextGameServerClientId());
+		ULocalPlayer* DummyPlayer = NewObject<ULocalPlayer>(GEngine, GEngine->LocalPlayerClass);
+		ProxyNetDriver->GetWorld()->GetGameInstance()->AddLocalPlayer(DummyPlayer, DummyClientId);
+
 		FURL URL = GameServerConnectionState->GameServerURL;
 		
-		const FString GameLoginOptions = NewPlayer->GetGameLoginOptions();
+		const FString GameLoginOptions = DummyPlayer->GetGameLoginOptions();
 		if (!GameLoginOptions.IsEmpty())
 		{
 			URL.AddOption(*GameLoginOptions);
@@ -308,24 +333,34 @@ void UProxyListenerNotify::ConnectToGameServer(UNetConnection* ProxyConnection,
 		GameServerConnectionState->GameServerNotify = NewObject<UGameServerNotify>();
 		GameServerConnectionState->GameServerNotify->Initialize(URL);
 		GameServerConnectionState->GameServerNotify->InitNetDriver(GameServerConnectionState->NetDriver);
-		GameServerConnectionState->GameServerNotify->SetFirstPlayer(NewPlayer);
+		GameServerConnectionState->GameServerNotify->SetFirstPlayer(DummyPlayer);
 		GameServerConnectionState->GameServerNotify->SetProxyNetDriver(ProxyNetDriver);
-		GameServerConnectionState->GameServerNotify->SetFlags(Flags);
+		GameServerConnectionState->GameServerNotify->SetFlags(UE::Net::EJoinFlags::NoPawn);		
 		GameServerConnection = GameServerConnectionState->NetDriver->ServerConnection;
 
 		// UNetDriver::Notify will be reset in UPendingNetGame above so it's important that we override
 		// it here again to point to the proxy.
 		GameServerConnectionState->NetDriver->Notify = GameServerConnectionState->GameServerNotify;
 
-		UE_LOG(LogNetProxy, Log, TEXT("Connecting to game server: %s:%s -> %s (player=%s client_handshake_id=%u)"), 
+		UE_LOG(LogNetProxy, Log, TEXT("Connecting to game server (parent): %s:%s -> %s (player=%s client_handshake_id=%u)"), 
 			   *ProxyConnection->GetDriver()->GetName(), 
 			   *ProxyConnection->GetName(),
 			   *GameServerConnection->GetDriver()->GetName(),
-			   *NewPlayer->GetName(),
+			   *DummyPlayer->GetName(),
 			   ClientHandshakeId);
 	}
-	else
+
 	{
+		// Add a player to use on the game server.
+		const FPlatformUserId GameServerClientId = FPlatformUserId::CreateFromInternalId(ProxyNetDriver->GetNextGameServerClientId());
+		ULocalPlayer* NewPlayer = NewObject<ULocalPlayer>(GEngine, GEngine->LocalPlayerClass);
+		ProxyNetDriver->GetWorld()->GetGameInstance()->AddLocalPlayer(NewPlayer, GameServerClientId);
+
+		// The new player will use the same unique identifier as the incoming proxy connection so that it will be propogated up to 
+		// the game servers through UNetConnection::PlayerId. This way each game server's incoming connection will have a PlayerId 
+		// that corresponds to a client connected to the proxy.
+		NewPlayer->SetCachedUniqueNetId(ProxyConnection->PlayerId);
+
 		FGameServerSplitJoinRequest Request;
 		Request.Player = NewPlayer;
 		Request.Flags = Flags;
@@ -355,46 +390,62 @@ void UProxyListenerNotify::ConnectToGameServer(UNetConnection* ProxyConnection,
 				   *NewPlayer->GetName(),
 				   Request.ClientHandshakeId);
 		}
-	}
-
-	GameServerConnectionState->Players.Add(NewPlayer);
 		
-	// Associate this proxy connection with the parent game server connection when beginning the handshake because
-	// the child connection hasn't been created yet. Once the handshake is complete, it's expected that this entry will
-	// be updated with the new child connection.
-	FMultiServerProxyInternalConnectionRoute Route;
-	Route.ClientHandshakeId = ClientHandshakeId;
-	Route.ProxyConnection = ProxyConnection;
-	Route.ParentGameServerConnection = GameServerConnection;
-	Route.GameServerConnection = nullptr;
-	Route.Player = NewPlayer;
-	Route.PlayerController = nullptr;
-	Route.LastUpdateViewTargetSec = 0.0;
-	Route.LastViewTargetPos = FVector(0, 0, 0);
-
-	if (!EnumHasAllFlags(Flags, UE::Net::EJoinFlags::NoPawn))
-	{
-		Route.State = EProxyConnectionState::ConnectingPrimary;
-	}
-	else
-	{
-		Route.State = EProxyConnectionState::Connecting;
-	}
+		GameServerConnectionState->Players.Add(NewPlayer);
+		
+		// Associate this proxy connection with the parent game server connection when beginning the handshake because
+		// the child connection hasn't been created yet. Once the handshake is complete, it's expected that this entry will
+		// be updated with the new child connection.
+		FMultiServerProxyInternalConnectionRoute Route;
+		Route.ClientHandshakeId = ClientHandshakeId;
+		Route.ProxyConnection = ProxyConnection;
+		Route.ParentGameServerConnection = GameServerConnection;
+		Route.GameServerConnection = nullptr;
+		Route.Player = NewPlayer;
+		Route.PlayerController = nullptr;
+		Route.LastUpdateViewTargetSec = 0.0;
+		Route.LastViewTargetPos = FVector(0, 0, 0);
+
+		if (!EnumHasAllFlags(Flags, UE::Net::EJoinFlags::NoPawn))
+		{
+			Route.State = EProxyConnectionState::ConnectingPrimary;
+		}
+		else
+		{
+			Route.State = EProxyConnectionState::Connecting;
+		}
 	
-	ProxyNetDriver->InternalRoutes.Add(ClientHandshakeId, Route);
+		ProxyNetDriver->InternalRoutes.Add(ClientHandshakeId, Route);
+	}
 }
 
 void UProxyBackendNetConnection::HandleClientPlayer(APlayerController* NewPlayerController, UNetConnection* GameServerConnection)
 {
-	// This function is called when a PlayerController is replicated to the proxy from a game server and represents the finalization
-	// of a connection to a primary or non-primary game server.
+	Super::HandleClientPlayer(NewPlayerController, GameServerConnection);
 
-	ensure(GameServerConnection == this);
+	NewPlayerController->NetConnection = GameServerConnection;
 
+	// Send any queued join requests for pending multiplexed connections.
 	UProxyBackendNetDriver* BackendNetDriver = Cast<UProxyBackendNetDriver>(Driver);
-	if (ensure(BackendNetDriver))
+	if (!ensure(BackendNetDriver))
 	{
-		BackendNetDriver->GetProxyNetDriver()->GameServerAssignPlayerController(this, GameServerConnection, NewPlayerController);
+		return;
+	}
+
+	UProxyNetDriver* ProxyNetDriver = Cast<UProxyNetDriver>(BackendNetDriver->GetProxyNetDriver());
+	if (ensure(BackendNetDriver) && ensure(ProxyNetDriver))
+	{
+		for (int32 GameServerIndex = 0; GameServerIndex < ProxyNetDriver->GetGameServerConnectionCount(); GameServerIndex++)
+		{
+			if (FGameServerConnectionState* GameServerConnectionState = ProxyNetDriver->GetGameServerConnection(GameServerIndex))
+			{
+				if (GameServerConnectionState->NetDriver->ServerConnection == GameServerConnection)
+				{
+					ProxyNetDriver->FlushSplitJoinRequests(GameServerConnectionState);
+					break;
+				}
+			}
+		}
 	}
 }
 
@@ -429,6 +480,7 @@ bool UProxyBackendNetDriver::InitBase(bool bInitAsClient, FNetworkNotify* InNoti
 	if (bSuccess)
 	{
 		NetConnectionClass = UProxyBackendNetConnection::StaticClass();
+		ChildNetConnectionClass = UProxyBackendChildNetConnection::StaticClass();
 
 		// Don't allow any RPCs received game servers to be executed on the proxy.
 		EnableExecuteRPCFunctions(false);
@@ -458,6 +510,22 @@ void UProxyBackendNetDriver::ForwardRemoteFunction(UObject* RootObject, UObject*
 	// Drop RPCs that don't come from a player controller on a primary game server.
 	if (APlayerController* PlayerController = Cast<APlayerController>(OwningActor))
 	{
+		// This can happen when a proxy receives a parent connection RPC from the game server. These should be ignored because no proxy
+		// client is associated with the parent connection and therefore it cannot forward the RPC.
+		if (PlayerController->GetNetConnection() && Cast<UProxyBackendChildNetConnection>(PlayerController->GetNetConnection()) == nullptr)
+		{
+			UE_LOG(LogNetProxy, Verbose, TEXT("RPC: Ignoring %s on %s from game server %s because it's coming from a parent connection %s."), *Function->GetName(), *RootObject->GetName(), *GetName(), *PlayerController->GetNetConnection()->GetName());
+			return;
+		}
+
+		// This can happen when a connection to a backend game server was closed on the proxy but the NMT_CloseChildConnection message hasn't
+		// be processed by the server yet and is still sending actor RPCs.
+		if (PlayerController->NetConnection == nullptr)
+		{
+			UE_LOG(LogNetProxy, Verbose, TEXT("RPC: Ignoring %s on %s from game server %s because player controller %s has a null network connection."), *Function->GetName(), *RootObject->GetName(), *GetName(), *GetNameSafe(PlayerController));
+			return;
+		}
+
 		const FMultiServerProxyInternalConnectionRoute* Route = ProxyNetDriver->GetRouteWithPlayerController(PlayerController);
 		if (!ensure(Route))
 		{
@@ -491,15 +559,6 @@ bool UProxyBackendNetDriver::ShouldSkipRepNotifies() const
 	return true;
 }
 
-UChildConnection* UProxyBackendNetDriver::CreateChild(UNetConnection* Parent)
-{
-	UChildConnection* Child = NewObject<UProxyBackendChildNetConnection>();
-	Child->InitChildConnection(this, Parent);
-	Parent->Children.Add(Child);
-	return Child;
-}
-
-
 void UProxyBackendNetDriver::InternalProcessRemoteFunction(
 	class AActor* Actor,
 	class UObject* SubObject,
@@ -638,6 +697,21 @@ FGameServerConnectionState* UProxyNetDriver::GetGameServerConnection(int32 Index
 	return nullptr;
 }
 
+int32 UProxyNetDriver::GetGameServerChildConnectionCount() const
+{
+	int32 ChildConnections = 0;
+
+	for (const FGameServerConnectionState& GameServerConnection : GameServerConnections)
+	{
+		if (GameServerConnection.NetDriver && GameServerConnection.NetDriver->ServerConnection)
+		{
+			ChildConnections += GameServerConnection.NetDriver->ServerConnection->Children.Num();
+		}
+	}
+
+	return ChildConnections;
+}
+
 // Given a comma-separated list of strings (IpAddress1:StartPort-EndPort,IpAddress2:StartPort-EndPort)
 // Return the list of actual servers addresses: { IpAddress1:StartPort, ..., IpAddress1:EndPort, IpAddress2:StartPort, ... }
 // This allows us to use short-hand when connecting to multiple servers.
@@ -700,6 +774,11 @@ bool UProxyNetDriver::InitBase(bool bInitAsClient, FNetworkNotify* InNotify, con
 
 	const bool bSuccess = Super::InitBase(bInitAsClient, ProxyListenerNotify, URL, bReuseAddressAndPort, Error);
 
+	if (bSuccess)
+	{
+		NetConnectionClass = UProxyNetConnection::StaticClass();
+	}
+
 	FString GameServerAddresses;
 	if (FParse::Value(FCommandLine::Get(), TEXT("-ProxyGameServers="), GameServerAddresses, false))
 	{
@@ -758,7 +837,6 @@ void UProxyNetDriver::SetupNetworkMetrics()
 	GetMetrics()->CreateInt(UE::MultiServerProxy::Metrics::AggregatedBackendInRate, 0);
 	GetMetrics()->CreateInt(UE::MultiServerProxy::Metrics::AggregatedBackendOutRate, 0);
 	GetMetrics()->CreateInt(UE::MultiServerProxy::Metrics::NumReassignPlayerController, 0);
-	GetMetrics()->CreateInt(UE::MultiServerProxy::Metrics::NumReplicatedPawns, 0);
 }
 
 bool UProxyNetDriver::InitConnect(class FNetworkNotify* InNotify, const FURL& ConnectURL, FString& Error)
@@ -890,7 +968,7 @@ void UProxyNetDriver::InternalProcessRemoteFunction(
 		if (UNetConnection* ProxyConnection = GetProxyConnectionFromPrimaryPlayer(Player))
 		{
 			UPlayer* Owner = Actor->GetNetOwningPlayerAnyRole();
-			UE_LOG(LogNetProxy, Verbose, TEXT("RPC: Forwarding %s on %s (owner:%s) to proxy connection %s:%s for player %s"), *Function->GetName(), *Actor->GetName(), *GetNameSafe(Owner), *ProxyConnection->Driver->GetName(), *ProxyConnection->GetName(), *Player->GetName());
+			UE_LOG(LogNetProxy, Verbose, TEXT("RPC: Forwarding %s on %s (owner:%s) to proxy connection %s:%s for player %s"), *GetNameSafe(Function), *GetNameSafe(Actor), *GetNameSafe(Owner), *ProxyConnection->Driver->GetName(), *GetNameSafe(ProxyConnection), *GetNameSafe(Player));
 			Super::InternalProcessRemoteFunction(Actor, SubObject, ProxyConnection, Function, Parms, OutParms, Stack, bIsServer);
 		}
 		else
@@ -936,6 +1014,8 @@ bool UProxyNetDriver::CanDowngradeActorRole(UNetConnection* ProxyConnection, AAc
 
 void UProxyNetDriver::Shutdown()
 {
+	Super::Shutdown();
+
 	for (FGameServerConnectionState& ConnectionState : GameServerConnections)
 	{
 		GEngine->DestroyNamedNetDriver(ConnectionState.World, ConnectionState.NetDriverName);
@@ -943,8 +1023,6 @@ void UProxyNetDriver::Shutdown()
 	}
 
 	GameServerConnections.Reset();
-
-	Super::Shutdown();
 }
 
 void UProxyNetDriver::TickFlush(float DeltaSeconds)
@@ -1018,25 +1096,81 @@ void UProxyNetDriver::GameServerAssignPlayerController(UNetConnection* ChildGame
 		case EProxyConnectionState::PendingReassign:
 			ProcessConnectionState_Reassign(Route, ChildGameServerConnection, GameServerPlayerController);
 			break;
+		case EProxyConnectionState::PendingClose:
+			ProcessConnectionState_PendingClose(Route, ChildGameServerConnection, GameServerPlayerController);
+			break;
 	}
+}
 
-	// If this was a parent connection send any join requests for any pending multiplexed connections.
-	if (ChildGameServerConnection == Route->ParentGameServerConnection)
+TSharedPtr<FNetGUIDCache> UProxyNetDriver::GetSharedBackendNetGuidCache()
+{
+	return SharedBackendNetGuidCache;
+}
+
+void UProxyNetDriver::HandleClosedProxyConnection(UNetConnection* ProxyConnection)
+{	
+	TArray<uint32> RoutesToRemove;
+
+	for (TPair<uint32, FMultiServerProxyInternalConnectionRoute>& ClientHandShakeIdToRoute : InternalRoutes)
 	{
-		for (FGameServerConnectionState& GameServerConnectionState : GameServerConnections)
+		const uint32 ClientHandshakeId = ClientHandShakeIdToRoute.Key;
+		const FMultiServerProxyInternalConnectionRoute* Route = &ClientHandShakeIdToRoute.Value;
+
+		if (Route->ProxyConnection == ProxyConnection)
 		{
-			if (GameServerConnectionState.NetDriver->ServerConnection == Route->ParentGameServerConnection)
-			{
-				FlushSplitJoinRequests(&GameServerConnectionState);
-				break;
-			}
+			RoutesToRemove.Add(ClientHandshakeId);
+		}
+	}
+
+	UE_LOG(LogNetProxy, Log, TEXT("Closing %d routes associated with proxy connection %s."), 
+		   RoutesToRemove.Num(),
+		   *UE::MultiServerProxy::Private::ConnectionToString(ProxyConnection));
+
+	for (uint32 ClientHandshakeIdToRemove : RoutesToRemove)
+	{
+		const FMultiServerProxyInternalConnectionRoute* Route = InternalRoutes.Find(ClientHandshakeIdToRemove);
+		if (!ensure(Route))
+		{
+			continue;
+		}
+
+		// If the connection route to the backend server is fully open, close it immediately, otherwise mark the route
+		// for closing once the connection handshake has been completed.
+		if (Route->State == EProxyConnectionState::Connected || Route->State == EProxyConnectionState::ConnectedPrimary)
+		{
+			CloseAndCleanUpConnectionRoute(Route);
+		}
+		else
+		{
+			UE_LOG(LogNetProxy, Log, TEXT("[%u] Deferring closing game server connection %s using player controller %s and player id %s."),
+				   Route->ClientHandshakeId,
+				   *UE::MultiServerProxy::Private::ConnectionToString(Route->GameServerConnection),
+				   *UE::MultiServerProxy::Private::PlayerControllerToString(Route->PlayerController),
+				   *GetNameSafe(Route->Player));
+
+			Route->State = EProxyConnectionState::PendingClose;
 		}
 	}
+	
+	ProxyConnectionToPrimaryGameServerPlayerController.Remove(ProxyConnection);
+	DeferredReassignment.Remove(ProxyConnection);
 }
 
-TSharedPtr<FNetGUIDCache> UProxyNetDriver::GetSharedBackendNetGuidCache()
+
+void UProxyNetDriver::CloseAndCleanUpConnectionRoute(const FMultiServerProxyInternalConnectionRoute* Route)
 {
-	return SharedBackendNetGuidCache;
+	UE_LOG(LogNetProxy, Log, TEXT("[%u] Closing game server connection %s using player controller %s and player id %s."),
+		   Route->ClientHandshakeId,
+		   *UE::MultiServerProxy::Private::ConnectionToString(Route->GameServerConnection),
+		   *UE::MultiServerProxy::Private::PlayerControllerToString(Route->PlayerController),
+		   *GetNameSafe(Route->Player));
+
+	InternalRoutes.Remove(Route->ClientHandshakeId);
+	UChildConnection* GameServerConnectionAsChild = Cast<UChildConnection>(Route->GameServerConnection);
+	if (ensure(GameServerConnectionAsChild))
+	{
+		GameServerConnectionAsChild->CloseAndRemoveChild(UE::Net::ECloseChildFlags::SendCloseMessage);
+	}
 }
 
 void UProxyNetDriver::FinalizeGameServerConnection(FMultiServerProxyInternalConnectionRoute* Route, UNetConnection* GameServerConnection, APlayerController* PlayerController)
@@ -1154,6 +1288,24 @@ void UProxyNetDriver::ProcessConnectionState_Reassign(const FMultiServerProxyInt
 	ReassignPlayerController(Route, GameServerConnection, PlayerController);
 }
 
+void UProxyNetDriver::ProcessConnectionState_PendingClose(FMultiServerProxyInternalConnectionRoute* Route, UNetConnection* GameServerConnection, APlayerController* PlayerController)
+{
+	// This state is expected to be called when a connection is closed before it's completed the connection handshake.
+
+	UE_LOG(LogNetProxy, Log, TEXT("[%u] Closing connection %s after handshake is completed using player controller %s and player id %s."),
+		   Route->ClientHandshakeId,
+		   *UE::MultiServerProxy::Private::ConnectionToString(GameServerConnection),
+		   *UE::MultiServerProxy::Private::PlayerControllerToString(PlayerController),
+		   *GetNameSafe(Route->Player));
+
+	// Perform any initialization of the route, connection and player controller that is required to cleanly close
+	// the connection.
+	Route->GameServerConnection = GameServerConnection;
+	GameServerConnection->PlayerController = PlayerController;
+
+	CloseAndCleanUpConnectionRoute(Route);
+}
+
 void UProxyNetDriver::UpdateProxyNetworkMetrics()
 {
 	extern ENGINE_API float GAverageFPS;
@@ -1215,7 +1367,6 @@ void UProxyNetDriver::ResetProxyNetworkMetrics()
 	GetMetrics()->SetInt(UE::MultiServerProxy::Metrics::NumForwardClientRPC, 0);
 	GetMetrics()->SetInt(UE::MultiServerProxy::Metrics::NumForwardBackendRPC, 0);
 	GetMetrics()->SetInt(UE::MultiServerProxy::Metrics::NumReassignPlayerController, 0);
-	GetMetrics()->SetInt(UE::MultiServerProxy::Metrics::NumReplicatedPawns, 0);
 }
 
 void UProxyNetDriver::ReceivedReassignedNoPawnPlayerController(const FMultiServerProxyInternalConnectionRoute* Route)
@@ -1405,8 +1556,6 @@ void UProxyNetDriver::PrepareStateForRelevancy()
 
 	// Relevancy calcluations are made using the position returned by AActor::GetActorLocation(). Since the proxy is not simulating
 	// player movement the actor's location must be manually set to the location in AActor::ReplicatedMovement.
-	int32 NumReplicatedPawns = 0;
-
 	for (const TSharedPtr<FNetworkObjectInfo>& ObjectInfo : GetNetworkObjectList().GetAllObjects())
 	{
 		APawn* Pawn = Cast<APawn>(ObjectInfo->Actor);
@@ -1416,11 +1565,7 @@ void UProxyNetDriver::PrepareStateForRelevancy()
 		}
 
 		Pawn->SetActorLocation(Pawn->GetReplicatedMovement().Location);
-
-		NumReplicatedPawns++;
 	}
-
-	GetMetrics()->SetInt(UE::MultiServerProxy::Metrics::NumReplicatedPawns, NumReplicatedPawns);
 }
 
 void UProxyNetDriver::UpdateLocalViewTarget(const APlayerController* PrimaryPlayerController)
@@ -1472,4 +1617,22 @@ void UProxyNetDriver::SetRemoteViewTarget(const ANoPawnPlayerController* PlayerC
 
 	UFunction* Function = ANoPawnPlayerController::StaticClass()->FindFunctionByName(TEXT("ServerSetViewTargetPosition"));
 	GameServerConnection->Driver->ProcessRemoteFunction(const_cast<ANoPawnPlayerController*>(PlayerController), Function, &ViewTargetPos, nullptr, nullptr);
+}
+
+void UProxyNetConnection::CleanUp()
+{
+	if (UE::MultiServerProxy::Private::bEnableDisconnectionSupport)
+	{
+		if (Driver)
+		{
+			UProxyNetDriver* ProxyNetDriver = Cast<UProxyNetDriver>(Driver);
+	
+			if (ensure(ProxyNetDriver))
+			{
+				ProxyNetDriver->HandleClosedProxyConnection(this);
+			}
+		}
+	}
+
+	Super::CleanUp();
 }
\ No newline at end of file
diff --git a/Engine/Plugins/Runtime/MultiServerReplication/Source/MultiServerReplication/Public/MultiServerProxy.h b/Engine/Plugins/Runtime/MultiServerReplication/Source/MultiServerReplication/Public/MultiServerProxy.h
--- a/Engine/Plugins/Runtime/MultiServerReplication/Source/MultiServerReplication/Public/MultiServerProxy.h
+++ b/Engine/Plugins/Runtime/MultiServerReplication/Source/MultiServerReplication/Public/MultiServerProxy.h
@@ -17,7 +17,8 @@ enum class EProxyConnectionState : uint8
 	ConnectingPrimary,
 	Connected,
 	ConnectedPrimary,
-	PendingReassign
+	PendingReassign,
+	PendingClose
 };
 
 /** 
@@ -280,7 +281,6 @@ public:
 	virtual bool InitBase(bool bInitAsClient, FNetworkNotify* InNotify, const FURL& URL, bool bReuseAddressAndPort, FString& Error);
 	virtual void ForwardRemoteFunction(UObject* RootObject, UObject* SubObject, UFunction* Function, void* Parms) override;
 	virtual bool ShouldSkipRepNotifies() const override;
-	virtual UChildConnection* CreateChild(UNetConnection* Parent) override;
 	virtual void InternalProcessRemoteFunction(class AActor* Actor, class UObject* SubObject, class UNetConnection* Connection, class UFunction* Function, void* Parms, FOutParmRec* OutParms, FFrame* Stack, bool bIsServer) override;
 	virtual bool ShouldClientDestroyActor(AActor* Actor, EChannelCloseReason CloseReason) const;
 
@@ -337,6 +337,9 @@ public:
 	int32 GetGameServerConnectionCount() const;
 	FGameServerConnectionState* GetGameServerConnection(int32 Index);
 
+	/** Return the number of child connections to the game servers. */
+	int32 GetGameServerChildConnectionCount() const;
+
 	/** UNetDriver interface function. */
 	virtual bool InitBase(bool bInitAsClient, FNetworkNotify* InNotify, const FURL& URL, bool bReuseAddressAndPort, FString& Error) override;
 	virtual void SetupNetworkMetrics() override;
@@ -364,17 +367,27 @@ public:
 
 	/** Return a pointer to the shared NetGUID cache used by backend network connections. */
 	TSharedPtr<FNetGUIDCache> GetSharedBackendNetGuidCache();
+	
+	/** Called when an incoming proxy connection has been closed and closes connections to backend servers. */
+	void HandleClosedProxyConnection(UNetConnection* ProxyConnection);
+	
+	/** Send any split join requests that have been queued up to the game server. */
+	void FlushSplitJoinRequests(FGameServerConnectionState* GameServerConnectionState);
 
 private:
 
 	friend class UProxyListenerNotify;
 
-	// Classes in EngineTest that test the behavior of UMultiServerProxy.
-	friend class FNetTestProxyRPCRouting;
+	// Classes in EngineTest that test the behavior of UProxyNetDriver.
 	friend class FNetTestProxyConnectionRouting;
+	friend class FNetTestProxyCloseSecondClientBeforeConnectedToBackend;
+	friend class FNetTestProxyCloseFirstClientBeforeBackendParentConnected;
 	
 	friend class UProxyBackendNetDriver;
 
+	/** Close a route's connection and perform any clean up. */
+	void CloseAndCleanUpConnectionRoute(const FMultiServerProxyInternalConnectionRoute* Route);
+
 	/** Called when the connection handshake to a game server is completed (i.e. the first player controller has been replicated on that connection). */
 	void FinalizeGameServerConnection(FMultiServerProxyInternalConnectionRoute* Route, UNetConnection* GameServerConnection, APlayerController* PlayerController);
 
@@ -386,6 +399,7 @@ private:
 	void ProcessConnectionState_Connected(const FMultiServerProxyInternalConnectionRoute* Route, UNetConnection* GameServerConnection, APlayerController* PlayerController);
 	void ProcessConnectionState_ConnectedPrimary(const FMultiServerProxyInternalConnectionRoute* Route, UNetConnection* GameServerConnection, APlayerController* PlayerController);
 	void ProcessConnectionState_Reassign(const FMultiServerProxyInternalConnectionRoute* Route, UNetConnection* GameServerConnection, APlayerController* PlayerController);
+	void ProcessConnectionState_PendingClose(FMultiServerProxyInternalConnectionRoute* Route, UNetConnection* GameServerConnection, APlayerController* PlayerController);
 
 	void UpdateProxyNetworkMetrics();
 	void AggregateBackendNetDriverMetrics();
@@ -410,9 +424,6 @@ private:
 
 	/** Set all of the configuration options that disable executing actor functionality or game specific code. */
 	void DisableActorLogicAndGameCode();
-	
-	/** Send any split join requests that have been queued up to the game server. */
-	void FlushSplitJoinRequests(FGameServerConnectionState* GameServerConnectionState);
 
 	/** Prepare state on the proxy to be used for relevancy when replicating to proxy connections. */
 	void PrepareStateForRelevancy();
@@ -474,3 +485,16 @@ private:
 	// As a client connects to the proxy, randomize PrimaryGameServerForNextClient to one of the known game servers.
 	bool bRandomizePrimaryGameServerForNextClient = false;
 };
+
+/** 
+* A network connection used by UProxyNetDriver.
+*/
+UCLASS()
+class MULTISERVERREPLICATION_API UProxyNetConnection : public UIpConnection
+{
+	GENERATED_BODY()
+
+public:
+
+	virtual void CleanUp() override;
+};
diff --git a/Engine/Source/Runtime/Core/Public/Misc/EngineNetworkCustomVersion.h b/Engine/Source/Runtime/Core/Public/Misc/EngineNetworkCustomVersion.h
--- a/Engine/Source/Runtime/Core/Public/Misc/EngineNetworkCustomVersion.h
+++ b/Engine/Source/Runtime/Core/Public/Misc/EngineNetworkCustomVersion.h
@@ -95,6 +95,7 @@ struct FEngineNetworkCustomVersion
 		JoinNoPawn = 40,						// Bump version to support NMT_JoinNoPawn and NMT_JoinNoPawnSplit control messages.
 		ClientHandshakeId = 41,					// Bump version to support serialization of APlayerController::ClientHandeshakeId.
 		PawnRemoteViewPitchTo16Bit = 42,		// Bump version to support higher precision PawnRemoteViewPitch
+		CloseChildConnection = 43,				// Bump version to support closing child connections with NMT_CloseChildConnection.
 
 		// -----<new versions can be added above this line>-------------------------------------------------
 		VersionPlusOne,
diff --git a/Engine/Source/Runtime/Engine/Classes/Engine/ChildConnection.h b/Engine/Source/Runtime/Engine/Classes/Engine/ChildConnection.h
--- a/Engine/Source/Runtime/Engine/Classes/Engine/ChildConnection.h
+++ b/Engine/Source/Runtime/Engine/Classes/Engine/ChildConnection.h
@@ -8,6 +8,17 @@
 #include "Net/Core/Connection/NetCloseResult.h"
 #include "ChildConnection.generated.h"
 
+namespace UE::Net
+{
+	enum class ECloseChildFlags
+	{
+		None 				= 0,
+		SendCloseMessage	= (1 << 0),	/* Send the NMT_CloseChildConnection control message to the remote endpoint (client/server). */
+		DestroyActors		= (1 << 1)	/* Destroy all actors associated with the child connection. */
+	};
+	ENUM_CLASS_FLAGS(ECloseChildFlags);
+}
+
 /**
  * Represents a secondary split screen connection that reroutes calls to the parent connection.
  */
@@ -85,6 +96,9 @@ public:
 	ENGINE_API virtual void GracefulClose(FNetCloseResult&& CloseReason) override;
 	ENGINE_API virtual void HandleClientPlayer(class APlayerController* PC, class UNetConnection* NetConnection) override;
 	ENGINE_API virtual void CleanUp() override;
+	
+	/** Close the child connection without closing the parent and any other child connections and remove it from it's parent. */
+	ENGINE_API void CloseAndRemoveChild(UE::Net::ECloseChildFlags Flags);
 
 	virtual void InitRemoteConnection(UNetDriver* InDriver, class FSocket* InSocket, const FURL& InURL, const class FInternetAddr& InRemoteAddr, EConnectionState InState, int32 InMaxPacket = 0, int32 InPacketOverhead = 0) override {}
 	virtual void InitLocalConnection(UNetDriver* InDriver, class FSocket* InSocket, const FURL& InURL, EConnectionState InState, int32 InMaxPacket = 0, int32 InPacketOverhead = 0) override {}
diff --git a/Engine/Source/Runtime/Engine/Classes/Engine/LocalPlayer.h b/Engine/Source/Runtime/Engine/Classes/Engine/LocalPlayer.h
--- a/Engine/Source/Runtime/Engine/Classes/Engine/LocalPlayer.h
+++ b/Engine/Source/Runtime/Engine/Classes/Engine/LocalPlayer.h
@@ -213,6 +213,27 @@ public:
 	DECLARE_EVENT_OneParam(ULocalPlayer, FOnPlayerControllerChanged, APlayerController* /*NewPC*/);
 	FOnPlayerControllerChanged& OnPlayerControllerChanged() { return OnPlayerControllerChangedEvent; }
 	
+	/** 
+	 * Set the connection identifier for this local user. 
+	 *
+	 * The connection identifier is a unique identifier for this local user on a client and is used to 
+	 * associate the local player that began a connection handshake with a player controller replicated
+	 * to the client when the connection handshake is complete.
+	 *
+	 * Only set this if you are connecting to a game server outside of the normal connection flow in the 
+	 * engine (e.g. you are sending NMT_Join manually to initiate a connection).
+	 * 
+	 * This is the expected flow for how the connection identifier should be used:
+	 *	* (Client) Set the connection identifier before sending NMT_Join, NMT_JoinSplit, NMT_JoinNoPawnSplit or NMT_JoinNoPawn.
+	 * 	* (Client) The connection identifier is sent with the join message as a URL argument.
+	 *	* (Server) Set the connection identifier on the player controller spawned on the game server.
+	 *	* (Client) Match the player controller connection identifier to the correct local user connection identifier.
+	 */
+	void SetConnectionIdentifier(const int32 Id);
+
+	/** Get the client connection identifier for this local user. */
+	int32 GetConnectionIdentifier() const;
+
 private:
 	TArray<FSceneViewStateReference> ViewStates;
 
@@ -222,6 +243,8 @@ private:
 
 	mutable FOnControllerIdChanged OnControllerIdChangedEvent;
 
+	int32 ConnectionIdentifier = 0;
+
 	/** The platform user this player is assigned to, could correspond to multiple input devices */
 	FPlatformUserId PlatformUserId;
 
diff --git a/Engine/Source/Runtime/Engine/Classes/Engine/NetConnection.h b/Engine/Source/Runtime/Engine/Classes/Engine/NetConnection.h
--- a/Engine/Source/Runtime/Engine/Classes/Engine/NetConnection.h
+++ b/Engine/Source/Runtime/Engine/Classes/Engine/NetConnection.h
@@ -327,6 +327,11 @@ private:
 	UPROPERTY()
 	uint32 InternalAck:1;					// Internally ack all packets, for 100% reliable connections.
 
+	UPROPERTY(transient)
+	TMap<int32, TObjectPtr<class UChildConnection>> ChildrenLookup;
+
+	int32 NextServerNetPlayerIndex = 1;
+
 public:
 	bool IsInternalAck() const { return bInternalAck; }
 	void SetInternalAck(bool bValue) 
@@ -1095,6 +1100,32 @@ public:
 		return 0;
 	}
 
+	/** Return the child connection that uses the provided NetPlayerIndex. */
+	ENGINE_API UChildConnection* GetChildConnectionWithNetPlayerIndex(const int32 NetPlayerIndex);
+
+	/** Return true if the provided child connection belongs to the parent and the NetPlayerIndex output to OutNetPlayerIndex. */
+	ENGINE_API bool GetNetPlayerIndex(UChildConnection* ChildConnection, int32* OutNetPlayerIndex);
+
+	/** 
+	 * Create and return a child connection and assign it a new NetPlayerIndex.
+	 *
+	 * The new NetPlayerIndex is output in the argument OutNetPlayerIndex.
+	 */
+	ENGINE_API UChildConnection* CreateChild(int32* OutNetPlayerIndex);
+
+	/**
+	 * Create and return a child connection and assign it a NetPlayerIndex provided as an argument.
+	 */
+	ENGINE_API UChildConnection* CreateChildWithNetPlayerIndex(int32 NetPlayerIndex);
+
+	/** 
+	 * Remove a child connection from internal data structures.
+	 *
+	 * This function won't actually close the child connection or perform any other clean-up. To close
+	 * the child connection call UChildConnection::CloseAndRemoveChild().
+	 */
+	ENGINE_API void RemoveChild(UChildConnection* ChildConnection);
+
 	/**
 	 * Return the platform specific FInternetAddr type, containing this connections address.
 	 * If nullptr is returned, connection is not added to MappedClientConnections, and can't receive net packets which depend on this.
diff --git a/Engine/Source/Runtime/Engine/Classes/Engine/NetDriver.h b/Engine/Source/Runtime/Engine/Classes/Engine/NetDriver.h
--- a/Engine/Source/Runtime/Engine/Classes/Engine/NetDriver.h
+++ b/Engine/Source/Runtime/Engine/Classes/Engine/NetDriver.h
@@ -1001,6 +1001,9 @@ public:
 	UPROPERTY()
 	TObjectPtr<UClass> NetConnectionClass;
 
+	UPROPERTY()
+	TObjectPtr<UClass> ChildNetConnectionClass;
+
 	UPROPERTY()
 	TObjectPtr<UClass> ReplicationDriverClass;
 
diff --git a/Engine/Source/Runtime/Engine/Classes/Engine/World.h b/Engine/Source/Runtime/Engine/Classes/Engine/World.h
--- a/Engine/Source/Runtime/Engine/Classes/Engine/World.h
+++ b/Engine/Source/Runtime/Engine/Classes/Engine/World.h
@@ -3801,9 +3801,10 @@ public:
 	 * @param UniqueId - unique net ID of the player (may be zeroed if no online subsystem or not logged in, e.g. a local game or LAN match)
 	 * @param Error (out) - if set, indicates that there was an error - usually is set to a property from which the calling code can look up the actual message
 	 * @param InNetPlayerIndex (optional) - the NetPlayerIndex to set on the PlayerController
+	 * @param InLocalPlayerIdentifier (optional) - the identifier of the local player on the client associated with this player controller.
 	 * @return the PlayerController that was spawned (may fail and return NULL)
 	 */
-	UE_API APlayerController* SpawnPlayActor(class UPlayer* Player, ENetRole RemoteRole, const FURL& InURL, const FUniqueNetIdRepl& UniqueId, FString& Error, uint8 InNetPlayerIndex = 0);
+	UE_API APlayerController* SpawnPlayActor(class UPlayer* Player, ENetRole RemoteRole, const FURL& InURL, const FUniqueNetIdRepl& UniqueId, FString& Error, uint8 InNetPlayerIndex = 0, int32 InLocalPlayerIdentifier = 0);
 	
 	/**
 	 * Try to find an acceptable non-colliding location to place TestActor as close to possible to PlaceLocation. Expects PlaceLocation to be a valid location inside the level.
diff --git a/Engine/Source/Runtime/Engine/Classes/GameFramework/PlayerController.h b/Engine/Source/Runtime/Engine/Classes/GameFramework/PlayerController.h
--- a/Engine/Source/Runtime/Engine/Classes/GameFramework/PlayerController.h
+++ b/Engine/Source/Runtime/Engine/Classes/GameFramework/PlayerController.h
@@ -2329,6 +2329,12 @@ public:
 	/** Associate this player controller with a connection using it's handshake cookie. */
 	ENGINE_API void SetClientHandshakeId(uint32 InClientHandshakeId);
 
+	/** Set the connection identifier of the local user on the client that will be using this player controller. */
+	ENGINE_API void SetLocalPlayerConnectionIdentifier(int32 Identifier);
+	
+	/** Get the connection identifier of the local user on the client that will be using this player controller. */
+	ENGINE_API int32 GetLocalPlayerConnectionIdentifier() const;
+
 private:
 
 	/** The player id associated with this player controller. */
@@ -2338,6 +2344,8 @@ private:
 	/** An optional client handshake id used by the connection associated with this player controller. */
 	UPROPERTY(DuplicateTransient)
 	uint32 ClientHandshakeId = 0;
+
+	int32 LocalPlayerConnectionIdentifier;
 };
 
 /** 
diff --git a/Engine/Source/Runtime/Engine/Private/DataChannel.cpp b/Engine/Source/Runtime/Engine/Private/DataChannel.cpp
--- a/Engine/Source/Runtime/Engine/Private/DataChannel.cpp
+++ b/Engine/Source/Runtime/Engine/Private/DataChannel.cpp
@@ -2100,6 +2100,9 @@ void UControlChannel::ReceivedBunch( FInBunch& Bunch )
 				case NMT_JoinNoPawnSplit:
 					FNetControlMessage<NMT_JoinNoPawnSplit>::Discard(Bunch);
 					break;
+				case NMT_CloseChildConnection:
+					FNetControlMessage<NMT_CloseChildConnection>::Discard(Bunch);
+					break;
 				default:
 					// if this fails, a case is missing above for an implemented message type
 					// or the connection is being sent potentially malformed packets
diff --git a/Engine/Source/Runtime/Engine/Private/GameInstance.cpp b/Engine/Source/Runtime/Engine/Private/GameInstance.cpp
--- a/Engine/Source/Runtime/Engine/Private/GameInstance.cpp
+++ b/Engine/Source/Runtime/Engine/Private/GameInstance.cpp
@@ -883,6 +883,11 @@ ULocalPlayer* UGameInstance::CreateLocalPlayer(FPlatformUserId UserId, FString&
 		}
 
 		NewPlayer = NewObject<ULocalPlayer>(GetEngine(), GetEngine()->LocalPlayerClass);
+		
+		static int32 NextLocalPlayerIdentifier = 0;
+		NewPlayer->SetConnectionIdentifier(NextLocalPlayerIdentifier);
+		NextLocalPlayerIdentifier++;
+		
 		InsertIndex = AddLocalPlayer(NewPlayer, UserId);
 		UWorld* CurrentWorld = GetWorld();
 		if (bSpawnPlayerController && InsertIndex != INDEX_NONE && CurrentWorld != nullptr)
diff --git a/Engine/Source/Runtime/Engine/Private/LevelActor.cpp b/Engine/Source/Runtime/Engine/Private/LevelActor.cpp
--- a/Engine/Source/Runtime/Engine/Private/LevelActor.cpp
+++ b/Engine/Source/Runtime/Engine/Private/LevelActor.cpp
@@ -1068,7 +1068,7 @@ bool UWorld::DestroyActor( AActor* ThisActor, bool bNetForce, bool bShouldModify
 	Player spawning.
 -----------------------------------------------------------------------------*/
 
-APlayerController* UWorld::SpawnPlayActor(UPlayer* NewPlayer, ENetRole RemoteRole, const FURL& InURL, const FUniqueNetIdRepl& UniqueId, FString& Error, uint8 InNetPlayerIndex)
+APlayerController* UWorld::SpawnPlayActor(UPlayer* NewPlayer, ENetRole RemoteRole, const FURL& InURL, const FUniqueNetIdRepl& UniqueId, FString& Error, uint8 InNetPlayerIndex, int32 InLocalPlayerIdentifier)
 {
 	Error = TEXT("");
 
@@ -1101,6 +1101,7 @@ APlayerController* UWorld::SpawnPlayActor(UPlayer* NewPlayer, ENetRole RemoteRol
 		NewPlayerController->NetPlayerIndex = InNetPlayerIndex;
 		NewPlayerController->SetRole(ROLE_Authority);
 		NewPlayerController->SetReplicates(RemoteRole != ROLE_None);
+		NewPlayerController->SetLocalPlayerConnectionIdentifier(InLocalPlayerIdentifier);
 		if (RemoteRole == ROLE_AutonomousProxy)
 		{
 			NewPlayerController->SetAutonomousProxy(true);
diff --git a/Engine/Source/Runtime/Engine/Private/LocalPlayer.cpp b/Engine/Source/Runtime/Engine/Private/LocalPlayer.cpp
--- a/Engine/Source/Runtime/Engine/Private/LocalPlayer.cpp
+++ b/Engine/Source/Runtime/Engine/Private/LocalPlayer.cpp
@@ -376,6 +376,8 @@ void ULocalPlayer::SendSplitJoin(TArray<FString>& Options, UNetDriver* NetDriver
 				URL.AddOption(*FString::Printf(TEXT("Name=%s"), *PlayerName));
 			}
 
+			URL.AddOption(*FString::Printf(TEXT("LocalPlayerId=%d"), ConnectionIdentifier));
+
 			// Send any game-specific url options for this player
 			FString GameUrlOptions = GetGameLoginOptions();
 			if (GameUrlOptions.Len() > 0)
@@ -1811,4 +1813,12 @@ void ULocalPlayer::CleanupViewState(FStringView MidParentRootPath /*= {}*/)
 	}
 }
 
-
+void ULocalPlayer::SetConnectionIdentifier(const int32 Id)
+{
+	ConnectionIdentifier = Id;
+}
+	
+int32 ULocalPlayer::GetConnectionIdentifier() const
+{
+	return ConnectionIdentifier;
+}
\ No newline at end of file
diff --git a/Engine/Source/Runtime/Engine/Private/Net/NetworkMetricsDefs.cpp b/Engine/Source/Runtime/Engine/Private/Net/NetworkMetricsDefs.cpp
--- a/Engine/Source/Runtime/Engine/Private/Net/NetworkMetricsDefs.cpp
+++ b/Engine/Source/Runtime/Engine/Private/Net/NetworkMetricsDefs.cpp
@@ -104,6 +104,9 @@ const FName RawPing("RawPing");
 const FName RawPingClientMin("RawPingClientMin");
 const FName RawPingClientMax("RawPingClientMax");
 const FName RawPingClientAvg("RawPingClientAvg");
+const FName NumPlayerControllersForReplication("NumPlayerControllersForReplication");
+const FName NumNoPawnPlayerControllersForReplication("NumNoPawnPlayerControllersForReplication");
+const FName NumPawnsForReplication("NumPawnsForReplication");
 
 }
 
diff --git a/Engine/Source/Runtime/Engine/Private/NetConnection.cpp b/Engine/Source/Runtime/Engine/Private/NetConnection.cpp
--- a/Engine/Source/Runtime/Engine/Private/NetConnection.cpp
+++ b/Engine/Source/Runtime/Engine/Private/NetConnection.cpp
@@ -6,6 +6,7 @@
 
 #include "Engine/NetConnection.h"
 #include "Engine/ReplicationDriver.h"
+#include "Engine/GameInstance.h"
 #include "EngineStats.h"
 #include "Net/Core/Trace/Private/NetTraceInternal.h"
 #include "UObject/Package.h"
@@ -1520,6 +1521,61 @@ void UChildConnection::CleanUp()
 	Driver = NULL;
 }
 
+void UChildConnection::CloseAndRemoveChild(UE::Net::ECloseChildFlags Flags)
+{
+	UE_LOG(LogNet, Log, TEXT("[%s:%s] Closing child connection %s."), *GetNameSafe(GetDriver()), *GetNameSafe(Parent), *GetName());
+	
+	if (EnumHasAnyFlags(Flags, UE::Net::ECloseChildFlags::SendCloseMessage))
+	{
+		if (PlayerController)
+		{
+			int32 NetPlayerIndexParam = static_cast<int32>(PlayerController->NetPlayerIndex);
+			FNetControlMessage<NMT_CloseChildConnection>::Send(Parent, NetPlayerIndexParam);
+
+			UE_LOG(LogNet, Log, TEXT("[%s:%s] Sending NMT_CloseChildConnection with NetPlayerIndex %d."), 
+				   *GetNameSafe(GetDriver()), *GetNameSafe(Parent), NetPlayerIndexParam);
+		}
+		else
+		{
+			UE_LOG(LogNet, Warning, TEXT("[%s:%s] Cannot send NMT_CloseChildConnection because there's no valid player controller."), 
+				   *GetNameSafe(GetDriver()), *GetNameSafe(Parent));
+		}
+	}
+
+	if (PlayerController)
+	{
+		if (ULocalPlayer* LocalPlayer = PlayerController->GetLocalPlayer())
+		{
+			if (UWorld* World = GetWorld())
+			{
+				if (UGameInstance* GameInstance = World->GetGameInstance())
+				{
+					// A local player may be removed from the game before the proxy driver is cleaned up in the case where the
+					// game is shutting down or closing.
+					const bool bIsRegisteredLocalPlayer = GameInstance->GetLocalPlayers().Find(LocalPlayer) != INDEX_NONE;
+
+					if (bIsRegisteredLocalPlayer)
+					{
+						UE_LOG(LogNet, Log, TEXT("[%s:%s] Remove local player %s."), *GetNameSafe(GetDriver()), *GetNameSafe(Parent), *GetNameSafe(LocalPlayer));
+
+						GameInstance->RemoveLocalPlayer(LocalPlayer);
+					}
+				}
+			}
+		}
+	}
+
+	if (Parent)
+	{
+		Parent->RemoveChild(this);
+		
+		if (EnumHasAnyFlags(Flags, UE::Net::ECloseChildFlags::DestroyActors))
+		{
+			CleanUp();
+		}
+	}
+}
+
 void UChildConnection::InitChildConnection(UNetDriver* InDriver, UNetConnection* InParent)
 {
 	Driver = InDriver;
@@ -5139,6 +5195,112 @@ void UNetConnection::HandleClientPlayer( APlayerController *PC, UNetConnection*
 	NotifyConnectionUpdated();
 }
 
+UChildConnection* UNetConnection::GetChildConnectionWithNetPlayerIndex(const int32 NetPlayerIndex)
+{
+	TObjectPtr<UChildConnection>* ChildConnectionPtr = ChildrenLookup.Find(NetPlayerIndex);
+	if (!ChildConnectionPtr)
+	{
+		return nullptr;
+	}
+
+	UChildConnection* ChildConnection = *ChildConnectionPtr;
+
+	ensure(!ChildConnection->PlayerController || (ChildConnection->PlayerController && ChildConnection->PlayerController->NetPlayerIndex == NetPlayerIndex));
+	return ChildConnection;
+}
+
+bool UNetConnection::GetNetPlayerIndex(UChildConnection* ChildConnection, int32* OutNetPlayerIndex)
+{
+	if (OutNetPlayerIndex == nullptr)
+	{
+		return false;
+	}
+
+	for (TPair<int32, TObjectPtr<UChildConnection>>& Pair : ChildrenLookup)
+	{
+		if (Pair.Value == ChildConnection)
+		{
+			*OutNetPlayerIndex = Pair.Key;
+			return true;
+		}
+	}
+
+	return false;
+}
+
+UChildConnection* UNetConnection::CreateChild(int32* OutNetPlayerIndex)
+{
+	const int32 NetPlayerIndex = NextServerNetPlayerIndex;
+	NextServerNetPlayerIndex++;
+
+	UE_LOG(LogNet, Log, TEXT("[%s:%s] Assigning NetPlayerIndex %d to a new child connection."), *GetNameSafe(GetDriver()), *GetName(), NetPlayerIndex);
+
+	UChildConnection* ChildConnection = CreateChildWithNetPlayerIndex(NetPlayerIndex);
+
+	if (OutNetPlayerIndex)
+	{
+		*OutNetPlayerIndex = NetPlayerIndex;
+	}
+
+	return ChildConnection;
+}
+
+UChildConnection* UNetConnection::CreateChildWithNetPlayerIndex(int32 NetPlayerIndex)
+{
+	// A NetPlayerIndex of 0 is reserved for the parent connection so the NetPlayerIndex for
+	// child connections should always begin at 1.
+	ensure(NetPlayerIndex > 0);
+
+	UChildConnection* ChildConnection = NewObject<UChildConnection>(GetTransientPackage(), GetDriver()->ChildNetConnectionClass);
+	ChildConnection->InitChildConnection(GetDriver(), this);
+	
+	UE_LOG(LogNet, Log, TEXT("[%s:%s] Creating child connection %s NetPlayerIndex %d"), *GetNameSafe(GetDriver()), *GetName(), *GetNameSafe(ChildConnection), NetPlayerIndex);
+	
+	ensure(Children.Find(ChildConnection) == INDEX_NONE);
+	Children.Add(ChildConnection);
+
+	ensure(!ChildrenLookup.Contains(NetPlayerIndex));
+	ChildrenLookup.Add(NetPlayerIndex, ChildConnection);
+
+	return ChildConnection;
+}
+
+void UNetConnection::RemoveChild(UChildConnection* ChildConnection)
+{
+	ensure(Children.Find(ChildConnection) != INDEX_NONE);
+
+	Children.Remove(ChildConnection);
+		
+	int32 NetPlayerIndexToRemove = -1;
+
+	if (PlayerController && ChildConnection->PlayerController)
+	{
+		NetPlayerIndexToRemove = ChildConnection->PlayerController->NetPlayerIndex;
+	}
+	else
+	{
+		for (TPair<int32, TObjectPtr<UChildConnection>>& Pair : ChildrenLookup)
+		{
+			if (Pair.Value == ChildConnection)
+			{
+				NetPlayerIndexToRemove = Pair.Key;
+				break;
+			}
+		}
+	}
+	
+	UE_LOG(LogNet, Log, TEXT("Removing child connection %s from parent connection %s using NetPlayerIndex %d."), *GetNameSafe(ChildConnection), *GetNameSafe(this), NetPlayerIndexToRemove);
+
+	if (NetPlayerIndexToRemove >= 0)
+	{
+		// A NetPlayerIndex of 0 is reserved for the parent connection so it should never be found in the lookup table.
+		ensure(NetPlayerIndexToRemove > 0);
+
+		ensure(ChildrenLookup.Contains(NetPlayerIndexToRemove) && *ChildrenLookup.Find(NetPlayerIndexToRemove) == ChildConnection);
+		ChildrenLookup.Remove(NetPlayerIndexToRemove);
+	}
+}
+
 void UChildConnection::HandleClientPlayer(APlayerController* PC, UNetConnection* NetConnection)
 {
 	ULocalPlayer* NewPlayer = NULL;
@@ -5148,11 +5310,11 @@ void UChildConnection::HandleClientPlayer(APlayerController* PC, UNetConnection*
 	}
 	else
 	{
-		// find the first player that doesn't already have a connection
-		uint8 CurrentIndex = 0;
-		for (FLocalPlayerIterator It(GEngine, Driver->GetWorld()); It; ++It, CurrentIndex++)
+		// Find the local player with the same split screen identifier as the new player controller.
+		for (FLocalPlayerIterator It(GEngine, Driver->GetWorld()); It; ++It)
 		{
-			if (CurrentIndex == PC->NetPlayerIndex)
+			ULocalPlayer* LocalPlayer = *It;
+			if (PC->GetLocalPlayerConnectionIdentifier() == LocalPlayer->GetConnectionIdentifier())
 			{
 				NewPlayer = *It;
 				break;
diff --git a/Engine/Source/Runtime/Engine/Private/NetDriver.cpp b/Engine/Source/Runtime/Engine/Private/NetDriver.cpp
--- a/Engine/Source/Runtime/Engine/Private/NetDriver.cpp
+++ b/Engine/Source/Runtime/Engine/Private/NetDriver.cpp
@@ -532,6 +532,13 @@ static FAutoConsoleVariableRef CVarEnableSupportForNoPawnConnection(
 	TEXT("Enable server support for NMT_JoinNoPawn and NMT_JoinNoPawnSplit control messages."),
 	ECVF_Default);
 
+static bool GbEnableActorCountInStats = false;
+static FAutoConsoleVariableRef CVarEnableActorCountInStats(
+	TEXT("net.EnableActorCountInStats"),
+	GbEnableActorCountInStats,
+	TEXT("Enable the tracking of pawns and player controller counts in the network metrics database."),
+	ECVF_Default);
+
 namespace UE::Net
 {
 	static FString GRequiredEncryptionNetDriverDefNames_Internal = TEXT("all");
@@ -1634,6 +1641,9 @@ bool UNetDriver::InitConnectionClass(void)
 			UE_LOG(LogNet, Error,TEXT("Failed to load class '%s'"),*NetConnectionClassName);
 		}
 	}
+
+	ChildNetConnectionClass = UChildConnection::StaticClass();
+
 	return NetConnectionClass != NULL;
 }
 
@@ -2024,6 +2034,15 @@ void UNetDriver::SetupNetworkMetrics()
 
 	// (Server only) The average RawPing from among all client connections each frame.
 	GetMetrics()->CreateFloat(UE::Net::Metric::RawPingClientAvg, 0.0f);
+
+	// The number of instances of APlayerController registered for replication, excluding ANoPawnPlayerController.
+	GetMetrics()->CreateInt(UE::Net::Metric::NumPlayerControllersForReplication, 0);
+
+	// The number of instances of ANoPawnPlayerController registered for replication.
+	GetMetrics()->CreateInt(UE::Net::Metric::NumNoPawnPlayerControllersForReplication, 0);
+
+	// The number of instances of APawn registered for replication.
+	GetMetrics()->CreateInt(UE::Net::Metric::NumPawnsForReplication, 0);
 }
 
 void UNetDriver::SetupNetworkMetricsListeners(bool bIsClient)
@@ -4755,12 +4774,7 @@ void UNetDriver::ForceActorRelevantNextUpdate(AActor* Actor)
 
 UChildConnection* UNetDriver::CreateChild(UNetConnection* Parent)
 {
-	UE_LOG(LogNet, Log, TEXT("Creating child connection with %s parent"), *Parent->GetName());
-	UChildConnection* Child = NewObject<UChildConnection>();
-	Child->InitChildConnection(this, Parent);
-	Parent->Children.Add(Child);
-
-	return Child;
+	return Parent->CreateChild(nullptr);
 }
 
 void UNetDriver::PostGarbageCollect()
@@ -8749,6 +8763,36 @@ void UNetDriver::UpdateNetworkStats()
 			NumFramesOverIncomingBunchTimeLimit = 0;
 			StatUpdateTime = CurrentRealtimeSeconds;
 			StatUpdateFrames = 0;
+
+			// Record the number of player controllers in the world and registered for replication.
+			if (GbEnableActorCountInStats)
+			{
+				int32 NumPlayerControllersForReplication = 0;
+				int32 NumNoPawnPlayerControllersForReplication = 0;
+				int32 NumPawnsForReplication = 0;
+
+				for (const TSharedPtr<FNetworkObjectInfo>& ObjectInfo : GetNetworkObjectList().GetActiveObjects())
+				{
+					AActor* Actor = ObjectInfo->Actor;
+
+					if (Actor->IsA(ANoPawnPlayerController::StaticClass()))
+					{
+						NumNoPawnPlayerControllersForReplication++;
+					}
+					else if (Actor->IsA(APlayerController::StaticClass()))
+					{
+						NumPlayerControllersForReplication++;
+					}
+					else if (Actor->IsA(APawn::StaticClass()))
+					{
+						NumPawnsForReplication++;
+					}
+				}
+
+				GetMetrics()->SetInt(UE::Net::Metric::NumPlayerControllersForReplication, NumPlayerControllersForReplication);
+				GetMetrics()->SetInt(UE::Net::Metric::NumNoPawnPlayerControllersForReplication, NumNoPawnPlayerControllersForReplication);
+				GetMetrics()->SetInt(UE::Net::Metric::NumPawnsForReplication, NumPawnsForReplication);
+			}
 		}
 		else
 		{
diff --git a/Engine/Source/Runtime/Engine/Private/PlayerController.cpp b/Engine/Source/Runtime/Engine/Private/PlayerController.cpp
--- a/Engine/Source/Runtime/Engine/Private/PlayerController.cpp
+++ b/Engine/Source/Runtime/Engine/Private/PlayerController.cpp
@@ -1409,6 +1409,11 @@ void APlayerController::OnActorChannelOpen(FInBunch& InBunch, UNetConnection* Co
 	{
 		InBunch << ClientHandshakeId;
 	}
+
+	if (InBunch.EngineNetVer() >= FEngineNetworkCustomVersion::CloseChildConnection)
+	{
+		InBunch << LocalPlayerConnectionIdentifier;
+	}
 	
 	if (Connection->Driver != NULL && Connection == Connection->Driver->ServerConnection)
 	{
@@ -1423,21 +1428,19 @@ void APlayerController::OnActorChannelOpen(FInBunch& InBunch, UNetConnection* Co
 		}
 		else
 		{
-			ensure(NetPlayerIndex > 0);
-			
-			int32 ChildIndex = int32(NetPlayerIndex) - 1;
-
-			// NetPlayerIndex represents an index UNetConnection::Children on both the client
-			// and server. This loop will support cases where two player controllers are processed
-			// by the client out of order (e.g. NetPlayerIndex=2 before NetPlayerIndex=1).
-			for (int32 Index = Connection->Children.Num(); Index < (ChildIndex + 1); Index++)
+			// Only create a new child connection if there isn't alreayd one associated with a NetPlayerIndex.
+			UChildConnection* ChildConnection = Connection->GetChildConnectionWithNetPlayerIndex(NetPlayerIndex);
+			if (!ChildConnection)
 			{
-				Connection->Driver->CreateChild(Connection);
+				ChildConnection = Connection->CreateChildWithNetPlayerIndex(NetPlayerIndex);
 			}
 
-			if (ensure(Connection->Children.IsValidIndex(ChildIndex)))
+			if (ensure(ChildConnection))
 			{
-				Connection->Children[ChildIndex]->HandleClientPlayer(this, Connection);
+				UE_LOG(LogNet, Log, TEXT("[%s:%s] Binding player controller %s to child connection %s"), 
+					   *GetNameSafe(Connection->GetDriver()), *GetNameSafe(Connection), *GetName(), *GetNameSafe(ChildConnection));
+
+				ChildConnection->HandleClientPlayer(this, Connection);
 			}
 		}
 	}
@@ -1453,6 +1456,7 @@ void APlayerController::OnSerializeNewActor(FOutBunch& OutBunch)
 	// serialize PlayerIndex as part of the initial bunch for PlayerControllers so they can be matched to the correct client-side viewport
 	OutBunch << NetPlayerIndex;
 	OutBunch << ClientHandshakeId;
+	OutBunch << LocalPlayerConnectionIdentifier;
 }
 
 void APlayerController::OnNetCleanup(UNetConnection* Connection)
@@ -1663,6 +1667,35 @@ void APlayerController::Destroyed()
 	PlayerInput = NULL;
 	CheatManager = NULL;
 
+#if UE_WITH_REMOTE_OBJECT_HANDLE
+	// Find any instances of ANoPawnPlayerController that were created when this player controller was migrated to a different server.
+	TArray<ANoPawnPlayerController*> NoPawnPlayerControllersDep;
+	for (FConstPlayerControllerIterator It = GetWorld()->GetPlayerControllerIterator(); It; ++It)
+	{
+		const TWeakObjectPtr<APlayerController>& Ptr = *It;
+		if (!Ptr.IsRemote())
+		{
+			ANoPawnPlayerController* Controller = Cast<ANoPawnPlayerController>(It->Get());
+			if (Controller && 
+				Controller != this &&
+				Controller->CachedConnectionPlayerId == CachedConnectionPlayerId &&
+				Controller->GetClientHandshakeId() == GetClientHandshakeId())
+			{
+				NoPawnPlayerControllersDep.Add(Controller);
+			}
+		}
+	}
+
+	// Remove and destroy any identifier instances of ANoPawnPlayerController.
+	for (ANoPawnPlayerController* Controller : NoPawnPlayerControllersDep)
+	{
+		UE_LOG(LogPlayerController, Log, TEXT("Destroying player controller %s which was used when %s was migrated away from this server."), *GetNameSafe(Controller), *GetName());
+
+		GetWorld()->RemoveController(Controller);
+		GetWorld()->DestroyActor(Controller, true);
+	}
+#endif
+
 	Super::Destroyed();
 }
 
@@ -5015,11 +5048,10 @@ void APlayerController::PostMigrate(const FUObjectMigrationContext& MigrationCon
 		{
 			if (UNetConnection* ParentConnection = ChildConnection->GetParentConnection())
 			{
-				int32 ConnectionNetPlayerIndex = ParentConnection->Children.Find(ChildConnection);
-				if (ensure(ConnectionNetPlayerIndex != INDEX_NONE))
-				{
-					return ConnectionNetPlayerIndex + 1;
-				}
+				int32 NetPlayerIndex = 0;
+				const bool bFoundNetPlayerIndex = ParentConnection->GetNetPlayerIndex(ChildConnection, &NetPlayerIndex);
+				ensure(bFoundNetPlayerIndex);
+				return NetPlayerIndex;
 			}
 		}
 
@@ -5070,6 +5102,10 @@ void APlayerController::PostMigrate(const FUObjectMigrationContext& MigrationCon
 				// The ANoPawnPlayerController actor has a low NetUpdateFrequency, so ForceNetUpdate() must be called because
 				// it's going to be detached from a connection.
 				NetConnection->PlayerController->ForceNetUpdate();
+
+				// Detach the previous player controller from the network connection so that any attempts to send an RPC
+				// on that controller to the client will not succeed.
+				NetConnection->PlayerController->NetConnection = nullptr;
 				
 				const int32 PrevNetPlayerIndex = NetConnection->PlayerController->NetPlayerIndex;
 				const int32 PrevHandshakeTransactionId = NetConnection->PlayerController->GetClientHandshakeId();
@@ -6656,6 +6692,16 @@ void APlayerController::ClientAckTimeDilation_Implementation(float TimeDilation,
 	}
 }
 
+void APlayerController::SetLocalPlayerConnectionIdentifier(int32 Identifier)
+{
+	LocalPlayerConnectionIdentifier = Identifier;
+}
+
+int32 APlayerController::GetLocalPlayerConnectionIdentifier() const
+{
+	return LocalPlayerConnectionIdentifier;
+}
+
 ANoPawnPlayerController::ANoPawnPlayerController(const FObjectInitializer& ObjectInitializer)
 	: Super(ObjectInitializer)
 {
diff --git a/Engine/Source/Runtime/Engine/Private/Tests/NetTestHelpers.cpp b/Engine/Source/Runtime/Engine/Private/Tests/NetTestHelpers.cpp
--- a/Engine/Source/Runtime/Engine/Private/Tests/NetTestHelpers.cpp
+++ b/Engine/Source/Runtime/Engine/Private/Tests/NetTestHelpers.cpp
@@ -4,6 +4,8 @@
 
 #include "CoreGlobals.h"
 #include "Engine/NetworkObjectList.h"
+#include "Engine/NetConnection.h"
+#include "Engine/ChildConnection.h"
 #include "Engine/Engine.h"
 #include "Engine/GameInstance.h"
 #include "Engine/World.h"
@@ -33,6 +35,55 @@
 namespace UE::Net
 {
 
+bool IsClientConnected(FTestWorldInstance& Instance)
+{
+	UNetDriver* NetDriver = Instance.GetNetDriver();
+	if (NetDriver && NetDriver->ServerConnection)
+	{
+		if (NetDriver->ServerConnection->GetConnectionState() == USOCK_Open)
+		{
+			return true;
+		}
+
+		for (UChildConnection* ChildConnection : NetDriver->ServerConnection->Children)
+		{
+			if (ChildConnection->GetConnectionState() != USOCK_Open)
+			{
+				return false;
+			}
+		}
+	}
+
+	return false;
+}
+
+bool IsServerConnected(FTestWorldInstance& Instance)
+{
+	UNetDriver* NetDriver = Instance.GetNetDriver();
+	if (NetDriver)
+	{
+		for (UNetConnection* Connection : NetDriver->ClientConnections)
+		{
+			if (Connection->GetConnectionState() != USOCK_Open)
+			{
+				return false;
+			}
+
+			for (UChildConnection* ChildConnection : Connection->Children)
+			{
+				if (ChildConnection->GetConnectionState() != USOCK_Open)
+				{
+					return false;
+				}
+			}
+		}
+
+		return true;
+	}
+
+	return false;
+}
+
 FTestWorldInstance FTestWorldInstance::CreateServer(const TCHAR* InURL, const EReplicationSystem ReplicationSystem)
 {
 	const bool bPrevShouldUseIrisReplication = ShouldUseIrisReplication();
@@ -70,7 +121,7 @@ FTestWorldInstance FTestWorldInstance::CreateServer(const TCHAR* InURL, const ER
 	return NewInstance;
 }
 
-FTestWorldInstance FTestWorldInstance::CreateClient(int32 ServerPort)
+FTestWorldInstance FTestWorldInstance::CreateClient(int32 ServerPort, int32 NumLocalPlayers)
 {
 	FGameInstancePIEParameters ClientParams;
 	ClientParams.bSimulateInEditor = false;
@@ -100,6 +151,12 @@ FTestWorldInstance FTestWorldInstance::CreateClient(int32 ServerPort)
 	FString ClientBrowseError;
 	GEngine->Browse(*ClientWorldContext, ClientURL, ClientBrowseError);
 
+	for (int32 LocalPlayerIndex = 1; LocalPlayerIndex < NumLocalPlayers; LocalPlayerIndex++)
+	{
+		FString Error;
+		NewInstance.GameInstance->CreateLocalPlayer(LocalPlayerIndex, Error, true);
+	}
+
 	return NewInstance;
 }
 
diff --git a/Engine/Source/Runtime/Engine/Private/World.cpp b/Engine/Source/Runtime/Engine/Private/World.cpp
--- a/Engine/Source/Runtime/Engine/Private/World.cpp
+++ b/Engine/Source/Runtime/Engine/Private/World.cpp
@@ -331,6 +331,25 @@ namespace UE::Private::World
 			}
 		}
 	}
+
+	int32 GetLocalPlayerIdentifier(FURL& InUrl)
+	{
+		if (const TCHAR* LocalPlayerIdStr = InUrl.GetOption(TEXT("LocalPlayerId="), nullptr))
+		{
+			TCHAR* EndStr = nullptr;
+			uint32 LocalPlayerId = FCString::Strtoi(LocalPlayerIdStr, &EndStr, 10);
+			if (LocalPlayerIdStr != EndStr)
+			{
+				return LocalPlayerId;
+			}
+			else
+			{
+				UE_LOG(LogNet, Warning, TEXT("Invalid local player id."));
+			}
+		}
+
+		return 0;
+	}
 }
 
 namespace UE::Gameplay::CVars
@@ -7057,6 +7076,26 @@ void UWorld::NotifyControlMessage(UNetConnection* Connection, uint8 MessageType,
 #endif
 		switch (MessageType)
 		{
+			case NMT_CloseChildConnection:
+			{
+				int32 NetPlayerIndexToClose;
+				if (FNetControlMessage<NMT_CloseChildConnection>::Receive(Bunch, NetPlayerIndexToClose))
+				{
+					UE_LOG(LogNet, Log, TEXT("Received NMT_CloseChildConnection with NetPlayerIndex %d on the client"), NetPlayerIndexToClose);
+
+					if (UChildConnection* ChildConnection = Connection->GetChildConnectionWithNetPlayerIndex(NetPlayerIndexToClose))
+					{		
+						// The client shouldn't destroy actors as it will be done by the server.
+						ChildConnection->CloseAndRemoveChild(UE::Net::ECloseChildFlags::None);
+					}
+					else
+					{
+						UE_LOG(LogNet, Warning, TEXT("Unable to find a child connection matching NetPlayerIndex %d"), NetPlayerIndexToClose);
+					}
+				}
+				break;
+			}
+
 			case NMT_Failure:
 			{
 				// our connection attempt failed for some reason, for example a synchronization mismatch (bad GUID, etc) or because the server rejected our join attempt (too many players, etc)
@@ -7132,6 +7171,25 @@ void UWorld::NotifyControlMessage(UNetConnection* Connection, uint8 MessageType,
 		
 		switch (MessageType)
 		{
+			case NMT_CloseChildConnection:
+			{
+				int32 NetPlayerIndexToClose;
+				if (FNetControlMessage<NMT_CloseChildConnection>::Receive(Bunch, NetPlayerIndexToClose))
+				{
+					UE_LOG(LogNet, Log, TEXT("Received NMT_CloseChildConnection with NetPlayerIndex %d on the server"), NetPlayerIndexToClose);
+
+					if (UChildConnection* ChildConnection = Connection->GetChildConnectionWithNetPlayerIndex(NetPlayerIndexToClose))
+					{		
+						ChildConnection->CloseAndRemoveChild(UE::Net::ECloseChildFlags::DestroyActors);
+					}
+					else
+					{
+						UE_LOG(LogNet, Warning, TEXT("Unable to find a child connection matching NetPlayerIndex %d"), NetPlayerIndexToClose);
+					}
+				}
+				break;
+			}
+
 			case NMT_Hello:
 			{
 				uint8 IsLittleEndian = 0;
@@ -7526,20 +7584,24 @@ void UWorld::NotifyControlMessage(UNetConnection* Connection, uint8 MessageType,
 						FActorSpawnParameters SpawnInfo;
 						SpawnInfo.ObjectFlags |= RF_Transient;
 						SpawnInfo.bDeferConstruction = true;
-					
-						UChildConnection* ChildConnection = NetDriver->CreateChild(Connection);
+
+						int32 NetPlayerIndex = 0;
+						UChildConnection* ChildConnection = Connection->CreateChild(&NetPlayerIndex);
 						ChildConnection->PlayerId = UniqueIdRepl;
 						
 						UE::Private::World::ParseAndSetClientHandshakeId(InURL, ChildConnection);
+						
+						const int32 LocalPlayerIdentifier = UE::Private::World::GetLocalPlayerIdentifier(InURL);
 
 						ChildConnection->PlayerController = GetWorld()->SpawnActor<APlayerController>(ANoPawnPlayerController::StaticClass(), FVector(), FRotator(), SpawnInfo);
 						ChildConnection->PlayerController->SetRole(ROLE_Authority);
 						ChildConnection->PlayerController->SetReplicates(true);
 						ChildConnection->PlayerController->SetAutonomousProxy(true);
 						ChildConnection->PlayerController->FinishSpawning(FTransform());
-						ChildConnection->PlayerController->NetPlayerIndex = Connection->Children.Num();
+						ChildConnection->PlayerController->NetPlayerIndex = NetPlayerIndex;
 						ChildConnection->PlayerController->NetConnection = ChildConnection;
 						ChildConnection->PlayerController->SetClientHandshakeId(ChildConnection->GetClientHandshakeId());
+						ChildConnection->PlayerController->SetLocalPlayerConnectionIdentifier(LocalPlayerIdentifier);
 
 						ChildConnection->OwningActor = ChildConnection->PlayerController;
 
@@ -7661,7 +7723,8 @@ void UWorld::PreLoginCompleteSplit(const FString& ErrorMsg, TWeakObjectPtr<UNetC
 	ULevel* CurrentLevel = PersistentLevel;
 #endif
 
-	UChildConnection* ChildConn = NetDriver->CreateChild(Connection);
+	int32 NetPlayerIndex = 0;
+	UChildConnection* ChildConn = Connection->CreateChild(&NetPlayerIndex);
 	ChildConn->PlayerId = SplitRequestUniqueIdRepl;
 	ChildConn->SetPlayerOnlinePlatformName(Connection->GetPlayerOnlinePlatformName());
 	ChildConn->RequestURL = SplitRequestURL;
@@ -7672,9 +7735,11 @@ void UWorld::PreLoginCompleteSplit(const FString& ErrorMsg, TWeakObjectPtr<UNetC
 
 	UE::Private::World::ParseAndSetClientHandshakeId(JoinSplitURL, ChildConn);
 
+	const int32 LocalPlayerIdentifier = UE::Private::World::GetLocalPlayerIdentifier(JoinSplitURL);
+
 	UE_LOG(LogNet, Log, TEXT("JOINSPLIT: Join request: URL=%s"), *JoinSplitURL.ToString());
 	FString SpawnErrorMsg;
-	APlayerController* PC = SpawnPlayActor(ChildConn, ROLE_AutonomousProxy, JoinSplitURL, ChildConn->PlayerId, SpawnErrorMsg, uint8(Connection->Children.Num()));
+	APlayerController* PC = SpawnPlayActor(ChildConn, ROLE_AutonomousProxy, JoinSplitURL, ChildConn->PlayerId, SpawnErrorMsg, NetPlayerIndex, LocalPlayerIdentifier);
 	if (PC == NULL)
 	{
 		// Failed to connect.
diff --git a/Engine/Source/Runtime/Engine/Public/Net/DataChannel.h b/Engine/Source/Runtime/Engine/Public/Net/DataChannel.h
--- a/Engine/Source/Runtime/Engine/Public/Net/DataChannel.h
+++ b/Engine/Source/Runtime/Engine/Public/Net/DataChannel.h
@@ -231,4 +231,12 @@ DEFINE_CONTROL_CHANNEL_MESSAGE(IrisNetRefHandleError, 30, UE::Net::ENetRefHandle
 // by default and must be explicitly enabled on the receiving server using the cvar net.EnableSupportForNoPawnConnection or calling
 // UNetDriver::SetSupportForNoPawnConnection() directly.
 DEFINE_CONTROL_CHANNEL_MESSAGE(JoinNoPawn, 31);
-DEFINE_CONTROL_CHANNEL_MESSAGE(JoinNoPawnSplit, 32, FString, FUniqueNetIdRepl);
\ No newline at end of file
+DEFINE_CONTROL_CHANNEL_MESSAGE(JoinNoPawnSplit, 32, FString, FUniqueNetIdRepl);
+
+// Close a child connection remotely without closing all other child and parent connections.
+//
+// The argument to this message is the value of NetPlayerIndex associated with the child connection that must be closed.
+//
+// It is expected that the local child connection will be cleaned up manually and does not depend on this message being 
+// successfully processed by the remote server.
+DEFINE_CONTROL_CHANNEL_MESSAGE(CloseChildConnection, 34, int32);
\ No newline at end of file
diff --git a/Engine/Source/Runtime/Engine/Public/Net/NetworkMetricsDefs.h b/Engine/Source/Runtime/Engine/Public/Net/NetworkMetricsDefs.h
--- a/Engine/Source/Runtime/Engine/Public/Net/NetworkMetricsDefs.h
+++ b/Engine/Source/Runtime/Engine/Public/Net/NetworkMetricsDefs.h
@@ -106,4 +106,7 @@ extern ENGINE_API const FName RawPing;
 extern ENGINE_API const FName RawPingClientMin;
 extern ENGINE_API const FName RawPingClientMax;
 extern ENGINE_API const FName RawPingClientAvg;
+extern ENGINE_API const FName NumPlayerControllersForReplication;
+extern ENGINE_API const FName NumNoPawnPlayerControllersForReplication;
+extern ENGINE_API const FName NumPawnsForReplication;
 }
diff --git a/Engine/Source/Runtime/Engine/Public/Tests/NetTestHelpers.h b/Engine/Source/Runtime/Engine/Public/Tests/NetTestHelpers.h
--- a/Engine/Source/Runtime/Engine/Public/Tests/NetTestHelpers.h
+++ b/Engine/Source/Runtime/Engine/Public/Tests/NetTestHelpers.h
@@ -5,11 +5,12 @@
 #include "HAL/Platform.h"
 #include "Engine/NetDriver.h"
 
+#include "Containers/Array.h"
 #include "GameFramework/PlayerController.h"
+#include "Kismet/GameplayStatics.h"
 
 class IConsoleVariable;
 class UGameInstance;
-class UWorld;
 class UReplicationSystem;
 class UObjectReplicationBridge;
 
@@ -20,6 +21,14 @@ struct FGameInstancePIEParameters;
 namespace UE::Net
 {
 
+template<class T>
+int32 NumActors(UWorld* World)
+{
+	TArray<AActor*> Actors;
+	UGameplayStatics::GetAllActorsOfClass(Cast<UObject>(World), T::StaticClass(), Actors);
+	return Actors.Num();
+}
+
 /*
  * Used to set a cvar and restore it to it's original value when destroyed
  */
@@ -111,7 +120,7 @@ struct FTestWorldInstance
 
 public:
 	ENGINE_API static FTestWorldInstance CreateServer(const TCHAR* InURL, const EReplicationSystem ReplicationSystem = EReplicationSystem::Default);
-	ENGINE_API static FTestWorldInstance CreateClient(int32 ServerPort);
+	ENGINE_API static FTestWorldInstance CreateClient(int32 ServerPort, int32 NumLocalPlayers = 1);
 	ENGINE_API static FTestWorldInstance CreateProxy();
 
 	ENGINE_API ~FTestWorldInstance();
@@ -314,6 +323,9 @@ inline bool FTestWorlds::TickAllUntil(const PredicateT& Predicate, float DeltaSe
 	return bPredicateResult;
 }
 
+ENGINE_API bool IsClientConnected(FTestWorldInstance& Instance);
+ENGINE_API bool IsServerConnected(FTestWorldInstance& Instance);
+
 } // end namespace UE::Net
 
 #endif //WITH_EDITOR
\ No newline at end of file
