MultiServerNode now more evenly balances the number of NetDrivers that are created for each server's peer connections.

diff --git a/Engine/Plugins/Runtime/MultiServerReplication/Source/MultiServerReplication/Private/MultiServerNode.cpp b/Engine/Plugins/Runtime/MultiServerReplication/Source/MultiServerReplication/Private/MultiServerNode.cpp
--- a/Engine/Plugins/Runtime/MultiServerReplication/Source/MultiServerReplication/Private/MultiServerNode.cpp
+++ b/Engine/Plugins/Runtime/MultiServerReplication/Source/MultiServerReplication/Private/MultiServerNode.cpp
@@ -158,22 +158,71 @@ bool UMultiServerNode::RegisterServer(const FMultiServerNodeCreateParams& Params
 		// If our local IP wasn't passed in, assume localhost.
 		FString ListenIp = Params.ListenIp.IsEmpty() ? TEXT("127.0.0.1") : Params.ListenIp;
 
-		for (const FString& PeerAddress : SortedPeerAddresses)
+		// Find our own index in the peer list
+		const int32 LocalIndex = SortedPeerAddresses.IndexOfByPredicate([&ListenIp, &Params](const FString& PeerAddress)
+			{
+				FURL PeerURL(nullptr, ToCStr(PeerAddress), ETravelType::TRAVEL_Absolute);
+				return PeerURL.Host.Equals(ListenIp) && PeerURL.Port == Params.ListenPort;
+			});
+
+		if (LocalIndex == INDEX_NONE)
 		{
-			FURL PeerURL(nullptr, ToCStr(PeerAddress), ETravelType::TRAVEL_Absolute);
-			
-			if (!PeerURL.Valid)
+			UE_LOG(LogMultiServerReplication, Error, TEXT("Failed to find this local server instance in PeerAddresses list. Address: %s:%hu. Not connecting to peers."), ToCStr(ListenIp), Params.ListenPort);
+			return false;
+		}
+
+		// Balance the number of client vs. server connections to minimize the difference in number of unique
+		// NetDrivers created by each peer. This helps minimize latency when iterating & processing each peer.
+
+		// Rules, where n is the total number of peers:
+		// If the number of peers is even:
+		//  -The first of peers in SortedPeerAddresses connect to ceil((n - 1) / 2) other servers.
+		//  -The second half of peers in SortedPeerAddresses connect to floor((n - 1) / 2) other servers.
+		//
+		// If the number of peers is odd:
+		//  -All peers connect to (n - 1) / 2 servers.
+		//
+		// All connections are made starting with the index directly above the local peer's own in the sorted list,
+		// wrapping around.
+		//
+		// For example, when there are 4 peers, each one would make connections as follows:
+		// 1 -> 2 3
+		// 2 -> 3 4
+		// 3 -> 4
+		// 4 -> 1
+		// 
+		// Or with 7 peers:
+		// 1 -> 2 3 4
+		// 2 -> 3 4 5
+		// 3 -> 4 5 6
+		// 4 -> 5 6 7
+		// 5 -> 6 7 1
+		// 6 -> 7 1 2
+		// 7 -> 1 2 3
+		//
+		const int32 CeilConnections = FMath::CeilToInt((Params.PeerAddresses.Num() - 1) / 2.0f);
+		const int32 FloorConnections = FMath::FloorToInt((Params.PeerAddresses.Num() - 1) / 2.0f);
+
+		// When the number of total peers is even, peers in the first half of the list will connect to one more
+		// peer than those in the second half of the list.
+		//
+		// When the number of total peers is odd, all peers will connect to the same number of other peers.
+		const int32 NumToConnectTo = LocalIndex < (Params.PeerAddresses.Num() / 2) ? CeilConnections : FloorConnections;
+
+		UE_LOG(LogMultiServerReplication, Log, TEXT("Connecting to %d other peers."), NumToConnectTo);
+
+		for (int32 RemotePeerCount = 0; RemotePeerCount < NumToConnectTo; ++RemotePeerCount)
+		{
+			const int32 RemotePeerIndex = (LocalIndex + RemotePeerCount + 1) % Params.PeerAddresses.Num();
+			if (!ensure(SortedPeerAddresses.IsValidIndex(RemotePeerIndex)))
 			{
-				UE_LOG(LogMultiServerReplication, Verbose, TEXT("Failed to parse peer address %s, not connecting."), ToCStr(PeerAddress));
+				UE_LOG(LogMultiServerReplication, Error, TEXT("Peer index %d is out of range. Number of peer addresses: %d."), RemotePeerIndex, Params.PeerAddresses.Num());
 				continue;
 			}
 
-			// Only connect to peers that are before this instance in the list to prevent redundant connections.
-			if (PeerURL.Host.Equals(ListenIp) && PeerURL.Port == Params.ListenPort)
-			{
-				break;
-			}
+			const FString& PeerAddress = SortedPeerAddresses[RemotePeerIndex];
 
+			UE_LOG(LogMultiServerReplication, Log, TEXT("Connecting to peer at address %s."), ToCStr(PeerAddress));
 			if (!PeerAddress.IsEmpty())
 			{
 				UMultiServerPeerConnection* Peer = NewObject<UMultiServerPeerConnection>(this);
diff --git a/Engine/Source/Runtime/Engine/Private/PacketHandlers/StatelessConnectHandlerComponent.cpp b/Engine/Source/Runtime/Engine/Private/PacketHandlers/StatelessConnectHandlerComponent.cpp
--- a/Engine/Source/Runtime/Engine/Private/PacketHandlers/StatelessConnectHandlerComponent.cpp
+++ b/Engine/Source/Runtime/Engine/Private/PacketHandlers/StatelessConnectHandlerComponent.cpp
@@ -332,6 +332,8 @@ namespace UE::Net
 	/** RESTART_RESPONSE_SIZE_BITS for EHandshakeVersion::Randomized */
 	static constexpr int32 VerRandomizedRestartResponseSizeBits = 419;
 
+	/** Initial value of LastClientSendTimestamp, indicating the client hasn't sent anything */
+	static constexpr double UninitializedLastClientSendTimestamp = -1.0;
 
 	const TCHAR* LexToString(EHandshakePacketType PacketType)
 	{
@@ -391,7 +393,7 @@ StatelessConnectHandlerComponent::StatelessConnectHandlerComponent()
 	, LastServerSequence(0)
 	, LastClientSequence(0)
 	, MinClientHandshakeVersion(static_cast<EHandshakeVersion>(UE::Net::CurrentHandshakeVersion))
-	, LastClientSendTimestamp(0.0)
+	, LastClientSendTimestamp(UE::Net::UninitializedLastClientSendTimestamp)
 	, LastChallengeTimestamp(0.0)
 	, LastRestartPacketTimestamp(0.0)
 	, LastSecretId(0)
@@ -728,7 +730,7 @@ void StatelessConnectHandlerComponent::SendToServer(EHandshakeVersion HandshakeV
 
 		Handler->SetRawSend(false);
 
-		LastClientSendTimestamp = FPlatformTime::Seconds();
+		LastClientSendTimestamp = Driver->GetElapsedTime();
 	}
 }
 
@@ -1084,9 +1086,10 @@ void StatelessConnectHandlerComponent::Incoming(FBitReader& Packet)
 						{
 							UNetConnection* ServerConn = (Driver != nullptr ? ToRawPtr(Driver->ServerConnection) : nullptr);
 							double LastNetConnPacketTime = (ServerConn != nullptr ? ServerConn->LastReceiveRealtime : 0.0);
+							const double DriverElapsedTime = (Driver != nullptr ? Driver->GetElapsedTime() : 0.0);
 
 							// The server may send multiple restart handshake packets, so have a 10 second delay between accepting them
-							bPassedDelayCheck = (CurrentTime - LastClientSendTimestamp) > 10.0;
+							bPassedDelayCheck = (DriverElapsedTime - LastClientSendTimestamp) > 10.0;
 
 							// Some clients end up sending packets duplicated over multiple IP's, triggering the restart handshake.
 							// Detect this by checking if any restart handshake requests have been received in roughly the last second
@@ -1821,9 +1824,12 @@ void StatelessConnectHandlerComponent::Tick(float DeltaTime)
 
 	if (Handler->Mode == UE::Handler::Mode::Client)
 	{
-		if (State != UE::Handler::Component::State::Initialized && LastClientSendTimestamp != 0.0)
+		if (State != UE::Handler::Component::State::Initialized && LastClientSendTimestamp != UninitializedLastClientSendTimestamp)
 		{
-			double LastSendTimeDiff = FPlatformTime::Seconds() - LastClientSendTimestamp;
+			ensureMsgf(Driver, TEXT("StatelessConnectHandlerComponent::Tick: Driver is null."));
+			const double CurrentTime = Driver ? Driver->GetElapsedTime() : 0.0;
+
+			const double LastSendTimeDiff = CurrentTime - LastClientSendTimestamp;
 
 			if (LastSendTimeDiff > UE::Net::HandshakeResendInterval)
 			{
