Iris support in the multiserver proxy.

* Disabling OnRep, RPC and calls to AActor::PreReplication().
* Swap local and remote roles on actors replicating from the proxy to clients.
* Associate a player controller with a proxy connection.

Also includes an unrelated fix to correctly set the connection identifier in the local player.

diff --git a/Engine/Plugins/Runtime/MultiServerReplication/Source/MultiServerReplication/MultiServerReplication.Build.cs b/Engine/Plugins/Runtime/MultiServerReplication/Source/MultiServerReplication/MultiServerReplication.Build.cs
--- a/Engine/Plugins/Runtime/MultiServerReplication/Source/MultiServerReplication/MultiServerReplication.Build.cs
+++ b/Engine/Plugins/Runtime/MultiServerReplication/Source/MultiServerReplication/MultiServerReplication.Build.cs
@@ -13,7 +13,8 @@ public class MultiServerReplication : ModuleRules
                 "CoreOnline",
                 "Engine",
 				"NetCore",
-				"MultiServerConfiguration"
+				"MultiServerConfiguration",
+				"IrisCore"
             }
         );
 
diff --git a/Engine/Plugins/Runtime/MultiServerReplication/Source/MultiServerReplication/Private/MultiServerProxy.cpp b/Engine/Plugins/Runtime/MultiServerReplication/Source/MultiServerReplication/Private/MultiServerProxy.cpp
--- a/Engine/Plugins/Runtime/MultiServerReplication/Source/MultiServerReplication/Private/MultiServerProxy.cpp
+++ b/Engine/Plugins/Runtime/MultiServerReplication/Source/MultiServerReplication/Private/MultiServerProxy.cpp
@@ -12,6 +12,10 @@
 #include "Net/NetworkMetricsDatabase.h"
 #include "Net/NetworkMetricsDefs.h"
 
+#include "Iris/ReplicationSystem/ReplicationSystem.h"
+#include "Net/Iris/ReplicationSystem/EngineReplicationBridge.h"
+#include "Net/Iris/ReplicationSystem/ReplicationSystemUtil.h"
+
 DEFINE_LOG_CATEGORY_STATIC(LogNetProxy, Log, All);
 
 namespace UE::MultiServerProxy::Private
@@ -308,18 +312,19 @@ void UProxyListenerNotify::ConnectToGameServer(UNetConnection* ProxyConnection,
 		UE_LOG(LogNetProxy, Log, TEXT("Created a Proxy Game Server NetDriver (name=%s, url=%s)"), *GameServerConnectionState->NetDriver->GetName(), *GameServerConnectionState->GameServerURL.ToString());
 	}
 
-	const uint32 ClientHandshakeId = ProxyNetDriver->GetNextClientHandshakeId();
-
 	UNetConnection* GameServerConnection = GameServerConnectionState->NetDriver->ServerConnection;
 	const bool bIsFirstGameServerConnection = (GameServerConnection == nullptr);
 	if (bIsFirstGameServerConnection)
 	{
+		const uint32 ParentClientHandshakeId = ProxyNetDriver->GetNextClientHandshakeId();
+
 		// The parent connection will not be associated with a player on the proxy and only exists so to use the control channel
 		// and to allow all players (child connections) to be closed without having to worry about closing a parent connection and
 		// promoting a child connection to parent connection.
 
 		const FPlatformUserId DummyClientId = FPlatformUserId::CreateFromInternalId(ProxyNetDriver->GetNextGameServerClientId());
 		ULocalPlayer* DummyPlayer = NewObject<ULocalPlayer>(GEngine, GEngine->LocalPlayerClass);
+		DummyPlayer->SetConnectionIdentifier(ParentClientHandshakeId);
 		ProxyNetDriver->GetWorld()->GetGameInstance()->AddLocalPlayer(DummyPlayer, DummyClientId);
 
 		FURL URL = GameServerConnectionState->GameServerURL;
@@ -330,7 +335,7 @@ void UProxyListenerNotify::ConnectToGameServer(UNetConnection* ProxyConnection,
 			URL.AddOption(*GameLoginOptions);
 		}
 
-		URL.AddOption(*FString::Printf(TEXT("HandshakeId=%u"), ClientHandshakeId));
+		URL.AddOption(*FString::Printf(TEXT("HandshakeId=%u"), ParentClientHandshakeId));
 
 		// Maybe this should be a control message since it changes the server setting (it's global driver setting).
 		URL.AddOption(TEXT("AutonomousAsSimulated"));
@@ -353,13 +358,16 @@ void UProxyListenerNotify::ConnectToGameServer(UNetConnection* ProxyConnection,
 			   *ProxyConnection->GetName(),
 			   *GameServerConnection->GetDriver()->GetName(),
 			   *DummyPlayer->GetName(),
-			   ClientHandshakeId);
+			   ParentClientHandshakeId);
 	}
 
 	{
+		const uint32 ChildClientHandshakeId = ProxyNetDriver->GetNextClientHandshakeId();
+
 		// Add a player to use on the game server.
 		const FPlatformUserId GameServerClientId = FPlatformUserId::CreateFromInternalId(ProxyNetDriver->GetNextGameServerClientId());
 		ULocalPlayer* NewPlayer = NewObject<ULocalPlayer>(GEngine, GEngine->LocalPlayerClass);
+		NewPlayer->SetConnectionIdentifier(ChildClientHandshakeId);
 		ProxyNetDriver->GetWorld()->GetGameInstance()->AddLocalPlayer(NewPlayer, GameServerClientId);
 
 		// The new player will use the same unique identifier as the incoming proxy connection so that it will be propogated up to 
@@ -370,7 +378,7 @@ void UProxyListenerNotify::ConnectToGameServer(UNetConnection* ProxyConnection,
 		FGameServerSplitJoinRequest Request;
 		Request.Player = NewPlayer;
 		Request.Flags = Flags;
-		Request.ClientHandshakeId = ClientHandshakeId;
+		Request.ClientHandshakeId = ChildClientHandshakeId;
 
 		// The NMT_JoinSplit message can only be sent when the parent connection is open.
 		if (GameServerConnectionState->NetDriver->ServerConnection->GetConnectionState() == USOCK_Open)
@@ -403,7 +411,7 @@ void UProxyListenerNotify::ConnectToGameServer(UNetConnection* ProxyConnection,
 		// the child connection hasn't been created yet. Once the handshake is complete, it's expected that this entry will
 		// be updated with the new child connection.
 		FMultiServerProxyInternalConnectionRoute Route;
-		Route.ClientHandshakeId = ClientHandshakeId;
+		Route.ClientHandshakeId = ChildClientHandshakeId;
 		Route.ProxyConnection = ProxyConnection;
 		Route.ParentGameServerConnection = GameServerConnection;
 		Route.GameServerConnection = nullptr;
@@ -421,7 +429,7 @@ void UProxyListenerNotify::ConnectToGameServer(UNetConnection* ProxyConnection,
 			Route.State = EProxyConnectionState::Connecting;
 		}
 	
-		ProxyNetDriver->InternalRoutes.Add(ClientHandshakeId, Route);
+		ProxyNetDriver->InternalRoutes.Add(ChildClientHandshakeId, Route);
 	}
 }
 
@@ -491,15 +499,15 @@ TObjectPtr<UProxyNetDriver> UProxyBackendNetDriver::GetProxyNetDriver()
 
 bool UProxyBackendNetDriver::InitBase(bool bInitAsClient, FNetworkNotify* InNotify, const FURL& URL, bool bReuseAddressAndPort, FString& Error)
 {
+	// Called before InitBase() because the values can be used during driver initialization.
+	EnableExecuteRPCFunctions(false);
+
 	bool bSuccess = Super::InitBase(bInitAsClient, InNotify, URL, bReuseAddressAndPort, Error);
 
 	if (bSuccess)
 	{
 		NetConnectionClass = UProxyBackendNetConnection::StaticClass();
 		ChildNetConnectionClass = UProxyBackendChildNetConnection::StaticClass();
-
-		// Don't allow any RPCs received game servers to be executed on the proxy.
-		EnableExecuteRPCFunctions(false);
 	}
 
 	SetReplicateTransactionally(false);
@@ -787,6 +795,11 @@ bool UProxyNetDriver::InitBase(bool bInitAsClient, FNetworkNotify* InNotify, con
 
 	ProxyListenerNotify = NewObject<UProxyListenerNotify>(GEngine, UProxyListenerNotify::StaticClass());
 	ProxyListenerNotify->SetProxyNetDriver(this);
+	
+	// Called before InitBase() because the values can be used during driver initialization.
+	DisableActorLogicAndGameCode();
+	EnableRoleSwappingOfAllActors(true);
+	EnableReplicationOfActorsWithAnyRole(true);
 
 	const bool bSuccess = Super::InitBase(bInitAsClient, ProxyListenerNotify, URL, bReuseAddressAndPort, Error);
 
@@ -829,8 +842,6 @@ bool UProxyNetDriver::InitBase(bool bInitAsClient, FNetworkNotify* InNotify, con
 	}
 
 	SetReplicateTransactionally(false);
-	
-	DisableActorLogicAndGameCode();
 
 	GuidCache = TSharedPtr<FNetGUIDCache>(new FProxyNetGUIDCache(this));
 
@@ -915,17 +926,6 @@ bool UProxyNetDriver::ShouldReplicateFunction(AActor* Actor, UFunction* Function
 	return false;
 }
 
-void UProxyNetDriver::NotifyActorChannelOpen(UActorChannel* Channel, AActor* Actor)
-{
-	Super::NotifyActorChannelOpen(Channel, Actor);
-
-	// The actor roles in the proxy must be replicated to the client as-is (i.e. the role on the client
-	// must be the same as the role in the proxy). Since the client will always swap roles when receiving
-	// replicated objects and the proxy is transparent to the client, the role is swapped on the proxy 
-	// before replicating.
-	SetRoleSwapOnReplicate(Actor, true);
-}
-
 void UProxyNetDriver::AddNetworkActor(AActor* Actor)
 {
 	// Ideally the proxy shouldn't spawn any actors since it's just used as a cache to pass state
@@ -937,12 +937,18 @@ void UProxyNetDriver::AddNetworkActor(AActor* Actor)
 	// loaded by the proxy and those replicated from the connected game servers. It's assumed that 
 	// the actors replicated from the game servers will not have a role of ROLE_Authority and will
 	// therefore be unaffacted by this code and replicate as normal.
-	if (Actor->GetIsReplicated())
+	//
+	// This is ignored for APlayerController because it's set to ROLE_Authority temporarily in order
+	// to correctly set it bReplicates flag to true when establishing a connection to a backend server.
+	if (!Actor->IsA(APlayerController::StaticClass()))
 	{
-		if (Actor->GetLocalRole() == ROLE_Authority)
+		if (Actor->GetIsReplicated())
 		{
-			Actor->SetReplicates(false);
-			Actor->SetRole(ROLE_None);
+			if (Actor->GetLocalRole() == ROLE_Authority)
+			{
+				Actor->SetReplicates(false);
+				Actor->SetRole(ROLE_None);
+			}
 		}
 	}
 
@@ -951,6 +957,16 @@ void UProxyNetDriver::AddNetworkActor(AActor* Actor)
 
 void UProxyNetDriver::RemoveNetworkActor(AActor* Actor)
 {
+	// Iris will normally stop replicating an actor when AActor::EndReplication() is called, but actor callbacks 
+	// are disabled on the proxy so it needs to be done here.
+	if (UReplicationSystem* DriverReplicationSystem = GetReplicationSystem())
+	{
+		if (UEngineReplicationBridge* Bridge = Cast<UEngineReplicationBridge>(DriverReplicationSystem->GetReplicationBridge()))
+		{
+			Bridge->StopReplicatingActor(Actor, FStopReplicatingActorParams(EEndPlayReason::Destroyed));
+		}
+	}
+
 	Super::RemoveNetworkActor(Actor);
 
 	// Remove this actor from the proxy netguid cache. On a normal dedicated server there's a one-to-one
@@ -1043,6 +1059,23 @@ void UProxyNetDriver::Shutdown()
 
 void UProxyNetDriver::TickFlush(float DeltaSeconds)
 {
+	// The function AssociatePlayerControllerWithProxyConnection() would ideally be called after a connection to a backend server
+	// has been established in UProxyBackendChildNetConnection::HandleClientPlayer(). This is not currently possible because 
+	// the actor is only added to the Iris replication system after HandleClientPlayer() is called (via OnActorChannelOpen()).
+	if (GetReplicationSystem())
+	{
+		for (TPair<TObjectPtr<UNetConnection>, TObjectPtr<const APlayerController>>& ProxyConnectionToPlayerController : ProxyConnectionToPrimaryGameServerPlayerController)
+		{
+			TObjectPtr<UNetConnection> ProxyConnection = ProxyConnectionToPlayerController.Key;
+			TObjectPtr<const APlayerController> PlayerController = ProxyConnectionToPlayerController.Value;
+
+			if (PlayerController->HasActorBegunPlay() || PlayerController->IsActorBeginningPlay())
+			{
+				AssociatePlayerControllerWithProxyConnection(ProxyConnection, const_cast<APlayerController*>(PlayerController.Get()));
+			}
+		}
+	}
+
 	// Update any metrics in the database that aren't tracked while ticking the frame.
 	UpdateProxyNetworkMetrics();
 
@@ -1201,6 +1234,15 @@ void UProxyNetDriver::FinalizeGameServerConnection(FMultiServerProxyInternalConn
 	PlayerController->NetConnection = GameServerConnection;
 	PlayerController->SetRole(ROLE_AutonomousProxy);
 
+	// Set the replication flag on the player controller to ensure that it can be added to the Iris replication system. Calling it from here assumes 
+	// that AActor::OnActorChannelOpen(), which ultimately calls FinalizeGameServerConnection(), is called before UEngineReplicationBridge::StartReplicatingActor().
+	//
+	// This only has an effect on Iris because an actor that doesn't have the replication flag set 
+	ENetRole PreviousRole = PlayerController->GetLocalRole();
+	PlayerController->SetRole(ROLE_Authority);
+	PlayerController->SetReplicates(true);
+	PlayerController->SetRole(PreviousRole);
+
 	GameServerConnection->SetConnectionState(EConnectionState::USOCK_Open);
 	GameServerConnection->SetClientHandshakeId(Route->ClientHandshakeId);
 	GameServerConnection->PlayerController = PlayerController;
@@ -1416,6 +1458,28 @@ void UProxyNetDriver::ReceivedReassignedNoPawnPlayerController(const FMultiServe
 	}
 }
 
+void UProxyNetDriver::AssociatePlayerControllerWithProxyConnection(UNetConnection* ProxyConnection, APlayerController* PlayerController)
+{
+	// If these ensures hit then the player controller won't be added to the replication system leading to a connectivity issues.
+	ensure(PlayerController->GetIsReplicated());
+	ensure(PlayerController->IsActorBeginningPlay() || PlayerController->HasActorBegunPlay());
+
+	if (UReplicationSystem* DriverReplicationSystem = GetReplicationSystem())
+	{
+		if (UObjectReplicationBridge* ObjectReplicationBridge = DriverReplicationSystem->GetReplicationBridgeAs<UObjectReplicationBridge>())
+		{
+			const uint32 ConnectionId = ProxyConnection->GetConnectionHandle().GetParentConnectionId();
+			const UE::Net::FNetRefHandle RefHandle = ObjectReplicationBridge->GetReplicatedRefHandle(PlayerController);
+			if (RefHandle.IsValid())
+			{
+				DriverReplicationSystem->SetReplicationEnabledForConnection(ProxyConnection->GetConnectionHandle().GetParentConnectionId(), true);
+				DriverReplicationSystem->SetOwningNetConnection(RefHandle, ConnectionId);
+				DriverReplicationSystem->SetReplicationConditionConnectionFilter(RefHandle, UE::Net::EReplicationCondition::RoleAutonomous, ConnectionId, true);
+			}
+		}	
+	}
+}
+
 void UProxyNetDriver::ReceivedReassignedGamePlayerController(const FMultiServerProxyInternalConnectionRoute* Route)
 {
 	FPlayerControllerReassignment& Reassignment = DeferredReassignment.FindOrAdd(Route->ProxyConnection);
diff --git a/Engine/Plugins/Runtime/MultiServerReplication/Source/MultiServerReplication/Public/MultiServerProxy.h b/Engine/Plugins/Runtime/MultiServerReplication/Source/MultiServerReplication/Public/MultiServerProxy.h
--- a/Engine/Plugins/Runtime/MultiServerReplication/Source/MultiServerReplication/Public/MultiServerProxy.h
+++ b/Engine/Plugins/Runtime/MultiServerReplication/Source/MultiServerReplication/Public/MultiServerProxy.h
@@ -346,7 +346,6 @@ public:
 	virtual bool InitConnect(class FNetworkNotify* InNotify, const FURL& ConnectURL, FString& Error) override;
 	virtual void ForwardRemoteFunction(UObject* RootObject, UObject* SubObject, UFunction* Function, void* Parms) override;
 	virtual bool ShouldReplicateFunction(AActor* Actor, UFunction* Function) const override;
-	virtual void NotifyActorChannelOpen(UActorChannel* Channel, AActor* Actor) override;
 	virtual void AddNetworkActor(AActor* Actor) override;
 	virtual void RemoveNetworkActor(AActor* Actor) override;
 	virtual bool ShouldCallRemoteFunction(UObject* Object, UFunction* Function, const FReplicationFlags& RepFlags) const;
@@ -405,6 +404,14 @@ private:
 	void AggregateBackendNetDriverMetrics();
 	void ResetProxyNetworkMetrics();
 
+	/** 
+	 * Perform logic that associates a player controller with an incoming proxy connection to enable correct replication to that connection.
+	 *
+	 * This logic is required to glue the player controller used to connect to a backend game server with an incoming proxy connection and
+	 * only applies to the Iris replication system.
+	 */
+	void AssociatePlayerControllerWithProxyConnection(UNetConnection* ProxyConnection, APlayerController* PlayerController);
+
 	/** 
 	 * Receive a new game player controller for an already open route.
 	 * 
diff --git a/Engine/Source/Runtime/Engine/Classes/Engine/NetDriver.h b/Engine/Source/Runtime/Engine/Classes/Engine/NetDriver.h
--- a/Engine/Source/Runtime/Engine/Classes/Engine/NetDriver.h
+++ b/Engine/Source/Runtime/Engine/Classes/Engine/NetDriver.h
@@ -1476,14 +1476,22 @@ protected:
 
 public:
 
-	/** (Generic Replicaton Only) Enable/disable the execution of received RPC functions. */
+	/** Enable/disable the execution of received RPC functions. */
 	ENGINE_API void EnableExecuteRPCFunctions(bool bEnable);
 	ENGINE_API bool IsExecuteRPCFunctionsEnabled() const;
 
-	/** (Generic Replication Only) Enable/disable calls to AActor::PreReplication(). */
+	/** Enable/disable calls to AActor::PreReplication(). */
 	ENGINE_API void EnablePreReplication(bool bEnable);
 	ENGINE_API bool IsPreReplicationEnabled() const;
 
+	/** Enable/disable swapping local and remote roles of all actors when replicating actors (overrides calls to SetRoleSwapOnReplicate()). */
+	ENGINE_API void EnableRoleSwappingOfAllActors(bool bEnable);
+	ENGINE_API bool IsRoleSwappingOnAllActorsEnabled() const;
+
+	/** (Iris replication only) Enable/disable replication of actors that have a local role other than ROLE_Authority. */
+	ENGINE_API void EnableReplicationOfActorsWithAnyRole(bool bEnable);
+	ENGINE_API bool IsReplicationOfActorsWithAnyRoleEnabled() const;
+
 	/** (Generic Replicaton Only) Treat actors with ROLE_AutonomousProxy as ROLE_SimulatedProxy when replicating state (applies to all client connections). */
 	ENGINE_API void ReplicateAutonomousAsSimulated(bool bEnabled);
 	ENGINE_API bool IsReplicatingAutonomousAsSimulated() const;
@@ -1495,6 +1503,8 @@ private:
 
 	bool bExecuteRPCFunctions = true;
 	bool bCallPreReplication = true;
+	bool bSwapRolesOfAllActors = false;
+	bool bReplicationOfActorsWithAnyRole = false;
 	bool bReplicatingAutonomousAsSimulated = false;
 
 	/** A set of replicators that have queued RPCs for transactions (server only) */
diff --git a/Engine/Source/Runtime/Engine/Private/Net/Iris/ReplicationSystem/EngineReplicationBridge.cpp b/Engine/Source/Runtime/Engine/Private/Net/Iris/ReplicationSystem/EngineReplicationBridge.cpp
--- a/Engine/Source/Runtime/Engine/Private/Net/Iris/ReplicationSystem/EngineReplicationBridge.cpp
+++ b/Engine/Source/Runtime/Engine/Private/Net/Iris/ReplicationSystem/EngineReplicationBridge.cpp
@@ -95,6 +95,14 @@ bool IsActorValidForIrisReplication(const AActor* Actor)
 
 void ActorReplicationBridgePreUpdateFunction(TArrayView<UObject*> Instances, const UObjectReplicationBridge* Bridge)
 {
+	if (const UEngineReplicationBridge* EngineReplicationBridge = Cast<const UEngineReplicationBridge>(Bridge))
+	{
+		if (!EngineReplicationBridge->AllowExecutePreReplication())
+		{
+			return;
+		}
+	}
+	
 	UNetDriver* NetDriver = CastChecked<const UEngineReplicationBridge>(Bridge)->GetNetDriver();
 	for (UObject* Instance : Instances)
 	{
@@ -305,10 +313,13 @@ UE::Net::FNetRefHandle UEngineReplicationBridge::StartReplicatingActor(AActor* A
 		return FNetRefHandle::GetInvalid();
 	}
 
-	if (Actor->GetLocalRole() != ROLE_Authority)
+	if (!AllowReplicationOfActorsWithAnyRole())
 	{
-		UE_LOG_BRIDGEID(LogIris, VeryVerbose, TEXT("Actor %s NetRole isn't Authority."), ToCStr(GetFullNameSafe(Actor)));
-		return FNetRefHandle::GetInvalid();
+		if (Actor->GetLocalRole() != ROLE_Authority)
+		{
+			UE_LOG_BRIDGEID(LogIris, VeryVerbose, TEXT("Actor %s NetRole isn't Authority."), ToCStr(GetFullNameSafe(Actor)));
+			return FNetRefHandle::GetInvalid();
+		}
 	}
 
 	if (Actor->IsActorBeingDestroyed() || !IsValid(Actor) || Actor->IsUnreachable())
@@ -751,7 +762,9 @@ void UEngineReplicationBridge::SetNetDriver(UNetDriver* const InNetDriver)
 				}) != nullptr;
 
 			checkf(bRequiredChannelIsConfigured, TEXT("ObjectReplication needs the netdriver channel %s to work. Add this channel to the netdriver channel definitions config"), *RequiredChannelName.ToString());
-		}
+
+		bAllowExecutePreReplication = InNetDriver->IsPreReplicationEnabled();
+		bAllowReplicationOfActorsWithAnyRole = InNetDriver->IsReplicationOfActorsWithAnyRoleEnabled();}
 	}
 }
 
@@ -1086,3 +1099,12 @@ void UEngineReplicationBridge::ConsumeNetMetrics(TArray<FAnalyticsEventAttribute
 
 	GetReplicationSystem()->ResetNetMetrics();
 }
+bool UEngineReplicationBridge::AllowExecutePreReplication() const
+{
+	return bAllowExecutePreReplication;
+}
+
+bool UEngineReplicationBridge::AllowReplicationOfActorsWithAnyRole() const
+{
+	return bAllowReplicationOfActorsWithAnyRole;
+}
diff --git a/Engine/Source/Runtime/Engine/Private/Net/Iris/ReplicationSystem/ReplicationSystemUtil.cpp b/Engine/Source/Runtime/Engine/Private/Net/Iris/ReplicationSystem/ReplicationSystemUtil.cpp
--- a/Engine/Source/Runtime/Engine/Private/Net/Iris/ReplicationSystem/ReplicationSystemUtil.cpp
+++ b/Engine/Source/Runtime/Engine/Private/Net/Iris/ReplicationSystem/ReplicationSystemUtil.cpp
@@ -231,8 +231,23 @@ void FReplicationSystemUtil::ForEachReplicationSystem(TFunctionRef<void(UReplica
 
 UReplicationSystem* FReplicationSystemUtil::GetReplicationSystem(const AActor* Actor)
 {
-	UNetDriver* NetDriver = Actor && Actor->GetWorld() ? Actor->GetNetDriver() : nullptr;
-	return NetDriver ? NetDriver->GetReplicationSystem() : nullptr;
+	if (!Actor)
+	{
+		return nullptr;
+	}
+
+	if (const UNetConnection* NetConnection = Actor->GetNetConnection())
+	{
+		return NetConnection->Driver->GetReplicationSystem();
+	}
+
+	UNetDriver* NetDriver = Actor->GetNetDriver();
+	if (!Actor->GetWorld() || !NetDriver)
+	{
+		return nullptr;
+	}
+
+	return NetDriver->GetReplicationSystem();
 }
 
 UReplicationSystem* FReplicationSystemUtil::GetReplicationSystem(const UNetDriver* NetDriver)
diff --git a/Engine/Source/Runtime/Engine/Private/NetDriver.cpp b/Engine/Source/Runtime/Engine/Private/NetDriver.cpp
--- a/Engine/Source/Runtime/Engine/Private/NetDriver.cpp
+++ b/Engine/Source/Runtime/Engine/Private/NetDriver.cpp
@@ -2300,6 +2300,26 @@ bool UNetDriver::IsPreReplicationEnabled() const
 	return bCallPreReplication;
 }
 
+void UNetDriver::EnableRoleSwappingOfAllActors(bool bEnable)
+{
+	bSwapRolesOfAllActors = bEnable;
+}
+
+bool UNetDriver::IsRoleSwappingOnAllActorsEnabled() const
+{
+	return bSwapRolesOfAllActors;
+}
+
+void UNetDriver::EnableReplicationOfActorsWithAnyRole(bool bEnable)
+{
+	bReplicationOfActorsWithAnyRole = bEnable;
+}
+
+bool UNetDriver::IsReplicationOfActorsWithAnyRoleEnabled() const
+{
+	return bReplicationOfActorsWithAnyRole;
+}
+
 void UNetDriver::ReplicateAutonomousAsSimulated(bool bEnabled)
 {
 	bReplicatingAutonomousAsSimulated = bEnabled;
@@ -5585,7 +5605,7 @@ int32 UNetDriver::ServerReplicateActors_ProcessPrioritizedActorsRange( UNetConne
 				FinalRelevantCount++;
 
 				TOptional<FScopedActorRoleSwap> SwapGuard;
-				if (ActorInfo->bSwapRolesOnReplicate)
+				if (ActorInfo->bSwapRolesOnReplicate || IsRoleSwappingOnAllActorsEnabled())
 				{
 					SwapGuard = FScopedActorRoleSwap(Actor);
 				}
@@ -7459,6 +7479,9 @@ void UNetDriver::CreateReplicationSystem(bool bInitAsClient)
 		Params.bUseRemoteObjectReferences = IsUsingRemoteObjectReferences();
 		Params.ForwardNetRPCCallDelegate.BindUObject(this, &UNetDriver::ForwardRemoteFunction);
 		UE::Net::Private::ApplyReplicationSystemConfig(RepSystemConfig, Params);
+		Params.bAllowExecuteRPC = IsExecuteRPCFunctionsEnabled();
+		Params.bAllowExecuteRepNotifies = !ShouldSkipRepNotifies();
+		Params.bAlwaysSwapRolesOnReplication = IsRoleSwappingOnAllActorsEnabled();
 		Params.NetTokenStore = NetTokenStore.Get();
 
 		SetReplicationSystem(UE::Net::FReplicationSystemFactory::CreateReplicationSystem(Params));
diff --git a/Engine/Source/Runtime/Engine/Public/Net/Iris/ReplicationSystem/EngineReplicationBridge.h b/Engine/Source/Runtime/Engine/Public/Net/Iris/ReplicationSystem/EngineReplicationBridge.h
--- a/Engine/Source/Runtime/Engine/Public/Net/Iris/ReplicationSystem/EngineReplicationBridge.h
+++ b/Engine/Source/Runtime/Engine/Public/Net/Iris/ReplicationSystem/EngineReplicationBridge.h
@@ -149,6 +149,9 @@ public:
 
 	/** Access to the factory id that handles subobjects */
 	UE::Net::FNetObjectFactoryId GetSubObjectFactoryId() const { return SubObjectFactoryId; }
+	
+	/** Return true if AActor::PreReplication() will be called on replicated objects (default: true). */
+	bool AllowExecutePreReplication() const;
 
 protected:
 	// UReplicationBridge
@@ -174,6 +177,9 @@ private:
 	
 	void OnMaxTickRateChanged(UNetDriver* InNetDriver, int32 NewMaxTickRate, int32 OldMaxTickRate);
 
+	/** Return true if actors that have a local role other than ROLE_Authority can be replicated to clients (default: false). */
+	bool AllowReplicationOfActorsWithAnyRole() const;
+
 private:
 
 	UE::Net::FNetObjectFactoryId ActorFactoryId;
@@ -191,4 +197,7 @@ private:
 	UPROPERTY(Config)
 	/** Name of the NetObjectFactory to use for default replicated subobjects */
 	FName SubObjectFactoryName;
+
+	bool bAllowExecutePreReplication = true;
+	bool bAllowReplicationOfActorsWithAnyRole = false;
 };
diff --git a/Engine/Source/Runtime/Net/Iris/Private/Iris/ReplicationSystem/DequantizeAndApplyHelper.cpp b/Engine/Source/Runtime/Net/Iris/Private/Iris/ReplicationSystem/DequantizeAndApplyHelper.cpp
--- a/Engine/Source/Runtime/Net/Iris/Private/Iris/ReplicationSystem/DequantizeAndApplyHelper.cpp
+++ b/Engine/Source/Runtime/Net/Iris/Private/Iris/ReplicationSystem/DequantizeAndApplyHelper.cpp
@@ -6,6 +6,7 @@
 #include "HAL/LowLevelMemStats.h"
 #include "Iris/Core/IrisProfiler.h"
 #include "Net/Core/NetBitArray.h"
+#include "Iris/ReplicationSystem/ReplicationSystem.h"
 #include "Iris/ReplicationState/ReplicationStateDescriptor.h"
 #include "Iris/ReplicationState/ReplicationStateUtil.h"
 #include "Iris/ReplicationSystem/ReplicationFragment.h"
@@ -13,6 +14,7 @@
 #include "Iris/ReplicationSystem/ReplicationProtocol.h"
 #include "Iris/Serialization/NetBitStreamReader.h"
 #include "Iris/Serialization/NetSerializer.h"
+#include "Iris/Serialization/InternalNetSerializationContext.h"
 #include "Misc/MemStack.h"
 #include "ProfilingDebugging/CsvProfiler.h"
 #include "ProfilingDebugging/AssetMetadataTrace.h"
@@ -254,20 +256,23 @@ void FDequantizeAndApplyHelper::CallLegacyPostApplyFunctions(FContext* Context,
 	CallLegacyFunctionForEachOwner(*Context, [](UObject* Object) { Object->PostNetReceive(); });
 
 	// CallRepNotifies
-	if (Context->bNeedsLegacyCallbacks)
+	if (NetSerializationContext.GetInternalContext()->ReplicationSystem->AllowExecuteRepNotify())
 	{
-		// We only call rep notifies for states that have received any data
-		for (const FContext::FStateData& StateData : MakeArrayView(Context->CachedStateData, Context->CachedStateCount))
+		if (Context->bNeedsLegacyCallbacks)
 		{
-			if (EnumHasAnyFlags(StateData.Fragment->GetTraits(), EReplicationFragmentTraits::HasRepNotifies))
+			// We only call rep notifies for states that have received any data
+			for (const FContext::FStateData& StateData : MakeArrayView(Context->CachedStateData, Context->CachedStateCount))
 			{
-				FReplicationStateApplyContext SetStateContext;
-				SetStateContext.NetSerializationContext = &NetSerializationContext;
-				SetStateContext.Descriptor = StateData.Descriptor;
-				SetStateContext.StateBufferData = StateData.StateBufferData;
-				SetStateContext.bIsInit = Context->bIsInitialState;
+				if (EnumHasAnyFlags(StateData.Fragment->GetTraits(), EReplicationFragmentTraits::HasRepNotifies))
+				{
+					FReplicationStateApplyContext SetStateContext;
+					SetStateContext.NetSerializationContext = &NetSerializationContext;
+					SetStateContext.Descriptor = StateData.Descriptor;
+					SetStateContext.StateBufferData = StateData.StateBufferData;
+					SetStateContext.bIsInit = Context->bIsInitialState;
 
-				StateData.Fragment->CallRepNotifies(SetStateContext);
+					StateData.Fragment->CallRepNotifies(SetStateContext);
+				}
 			}
 		}
 	}
diff --git a/Engine/Source/Runtime/Net/Iris/Private/Iris/ReplicationSystem/NetBlob/NetRPCHandler.cpp b/Engine/Source/Runtime/Net/Iris/Private/Iris/ReplicationSystem/NetBlob/NetRPCHandler.cpp
--- a/Engine/Source/Runtime/Net/Iris/Private/Iris/ReplicationSystem/NetBlob/NetRPCHandler.cpp
+++ b/Engine/Source/Runtime/Net/Iris/Private/Iris/ReplicationSystem/NetBlob/NetRPCHandler.cpp
@@ -4,6 +4,7 @@
 #include "Iris/ReplicationSystem/NetBlob/NetRPC.h"
 #include "Iris/ReplicationSystem/ReplicationSystem.h"
 #include "Iris/ReplicationSystem/ReplicationSystemInternal.h"
+#include "Iris/Serialization/InternalNetSerializationContext.h"
 #include "HAL/IConsoleManager.h"
 
 #include UE_INLINE_GENERATED_CPP_BY_NAME(NetRPCHandler)
@@ -44,6 +45,11 @@ TRefCountPtr<UE::Net::FNetBlob> UNetRPCHandler::CreateNetBlob(const FNetBlobCrea
 
 void UNetRPCHandler::OnNetBlobReceived(UE::Net::FNetSerializationContext& NetContext, const TRefCountPtr<FNetBlob>& NetBlob)
 {
+	if (!NetContext.GetInternalContext()->ReplicationSystem->AllowExecuteRPC())
+	{
+		return;
+	}
+
 	const UE::Net::FForwardNetRPCCallMulticastDelegate& ForwardNetRPCCallDelegate = ReplicationSystem->GetReplicationSystemInternal()->GetForwardNetRPCCallMulticastDelegate();
 	UE::Net::FNetRPCCallContext CallContext(NetContext, ForwardNetRPCCallDelegate);
 	FNetRPC* RPC = static_cast<FNetRPC*>(NetBlob.GetReference());
diff --git a/Engine/Source/Runtime/Net/Iris/Private/Iris/ReplicationSystem/ReplicationSystem.cpp b/Engine/Source/Runtime/Net/Iris/Private/Iris/ReplicationSystem/ReplicationSystem.cpp
--- a/Engine/Source/Runtime/Net/Iris/Private/Iris/ReplicationSystem/ReplicationSystem.cpp
+++ b/Engine/Source/Runtime/Net/Iris/Private/Iris/ReplicationSystem/ReplicationSystem.cpp
@@ -824,6 +824,9 @@ void UReplicationSystem::Init(uint32 InId, const FReplicationSystemParams& Param
 	Id = InId;
 	bIsServer = Params.bIsServer;
 	bAllowObjectReplication = Params.bAllowObjectReplication;
+	bAllowExecuteRPC = Params.bAllowExecuteRPC;
+	bAllowExecuteRepNotify = Params.bAllowExecuteRepNotifies;
+	bAlwaysSwapRolesOnReplication = Params.bAlwaysSwapRolesOnReplication;
 
 	ReplicationBridge = Params.ReplicationBridge;
 
diff --git a/Engine/Source/Runtime/Net/Iris/Private/Iris/ReplicationSystem/ReplicationWriter.cpp b/Engine/Source/Runtime/Net/Iris/Private/Iris/ReplicationSystem/ReplicationWriter.cpp
--- a/Engine/Source/Runtime/Net/Iris/Private/Iris/ReplicationSystem/ReplicationWriter.cpp
+++ b/Engine/Source/Runtime/Net/Iris/Private/Iris/ReplicationSystem/ReplicationWriter.cpp
@@ -2552,6 +2552,7 @@ FReplicationWriter::EWriteObjectStatus FReplicationWriter::WriteObjectAndSubObje
 			{
 				FInternalNetSerializationContext* InternalContext = Context.GetInternalContext();
 				InternalContext->bDowngradeAutonomousProxyRole = (Context.GetLocalConnectionId() != ReplicationFiltering->GetOwningConnection(InternalIndex));
+				InternalContext->bAlwaysSwapRolesOnReplication = InternalContext->ReplicationSystem->AlwaysSwapRolesOnReplication();
 			}
 
 			if (Writer.WriteBool(bIsInitialState))
diff --git a/Engine/Source/Runtime/Net/Iris/Private/Iris/Serialization/InternalNetSerializationContext.cpp b/Engine/Source/Runtime/Net/Iris/Private/Iris/Serialization/InternalNetSerializationContext.cpp
--- a/Engine/Source/Runtime/Net/Iris/Private/Iris/Serialization/InternalNetSerializationContext.cpp
+++ b/Engine/Source/Runtime/Net/Iris/Private/Iris/Serialization/InternalNetSerializationContext.cpp
@@ -15,6 +15,7 @@ FInternalNetSerializationContext::FInternalNetSerializationContext(UReplicationS
 , bDowngradeAutonomousProxyRole(0)
 , bInlineObjectReferenceExports(0)
 , bSerializeObjectReferencesAsRemoteIds(InReplicationSystem->IsUsingRemoteObjectReferences() ? 1 : 0)
+, bAlwaysSwapRolesOnReplication(0)
 {
 }
 
diff --git a/Engine/Source/Runtime/Net/Iris/Private/Iris/Serialization/InternalNetSerializationContext.h b/Engine/Source/Runtime/Net/Iris/Private/Iris/Serialization/InternalNetSerializationContext.h
--- a/Engine/Source/Runtime/Net/Iris/Private/Iris/Serialization/InternalNetSerializationContext.h
+++ b/Engine/Source/Runtime/Net/Iris/Private/Iris/Serialization/InternalNetSerializationContext.h
@@ -55,6 +55,9 @@ public:
 	// If true, object references will be quantized & serialized as FRemoteObjectIds instead of FNetObjectReferences.
 	// Requires UE_WITH_REMOTE_OBJECT_HANDLE.
 	uint32 bSerializeObjectReferencesAsRemoteIds : 1;
+
+	// If true, swap local and remote roles when replicating.
+	uint32 bAlwaysSwapRolesOnReplication : 1;
 };
 
 inline FInternalNetSerializationContext::FInternalNetSerializationContext()
diff --git a/Engine/Source/Runtime/Net/Iris/Private/Iris/Serialization/NetRoleNetSerializer.cpp b/Engine/Source/Runtime/Net/Iris/Private/Iris/Serialization/NetRoleNetSerializer.cpp
--- a/Engine/Source/Runtime/Net/Iris/Private/Iris/Serialization/NetRoleNetSerializer.cpp
+++ b/Engine/Source/Runtime/Net/Iris/Private/Iris/Serialization/NetRoleNetSerializer.cpp
@@ -42,7 +42,13 @@ void FNetRoleNetSerializer::Serialize(FNetSerializationContext& Context, const F
 {
 	const ConfigType* Config = static_cast<const ConfigType*>(Args.NetSerializerConfig);
 	const SourceType Value = *reinterpret_cast<const SourceType*>(Args.Source);
-	if (Context.GetInternalContext()->bDowngradeAutonomousProxyRole & (Value == Config->AutonomousProxyValue))
+	if (Context.GetInternalContext()->bAlwaysSwapRolesOnReplication)
+	{
+		FNetSerializeArgs SwapArgs = Args;
+		SwapArgs.Source += Config->RelativeInternalOffsetToOtherRole;
+		Private::FIntRangeNetSerializerBase<SourceType, ConfigType>::Serialize(Context, SwapArgs);
+	}
+	else if (Context.GetInternalContext()->bDowngradeAutonomousProxyRole & (Value == Config->AutonomousProxyValue))
 	{
 		FNetSerializeArgs DowngradedRoleArgs = Args;
 		DowngradedRoleArgs.Source = NetSerializerValuePointer(&Config->SimulatedProxyValue);
diff --git a/Engine/Source/Runtime/Net/Iris/Public/Iris/ReplicationSystem/ReplicationSystem.h b/Engine/Source/Runtime/Net/Iris/Public/Iris/ReplicationSystem/ReplicationSystem.h
--- a/Engine/Source/Runtime/Net/Iris/Public/Iris/ReplicationSystem/ReplicationSystem.h
+++ b/Engine/Source/Runtime/Net/Iris/Public/Iris/ReplicationSystem/ReplicationSystem.h
@@ -144,6 +144,15 @@ public:
 
 		/** If true we will only do a minimal update if there are no valid connections */
 		bool bAllowMinimalUpdateIfNoConnections = true;
+		
+		/** If true, call RepNotifies functions associated with replicated properties. */
+		bool bAllowExecuteRepNotifies = true;
+
+		/** If true, call RPC functions associated with replicated objects. */
+		bool bAllowExecuteRPC = true;
+
+		/** If true, swap local and remote roles before replicating ENetRole properties to clients. */
+		bool bAlwaysSwapRolesOnReplication = false;
 
 		/** Delegate that receives every RPC executed locally. */
 		UE::Net::FForwardNetRPCCallDelegate ForwardNetRPCCallDelegate;
@@ -161,6 +170,10 @@ public:
 	/** @return Whether the system is run on a server. */
 	bool IsServer() const { return bIsServer; }
 
+	bool AllowExecuteRPC() const { return bAllowExecuteRPC; }
+	bool AllowExecuteRepNotify() const { return bAllowExecuteRepNotify; };
+	bool AlwaysSwapRolesOnReplication() const { return bAlwaysSwapRolesOnReplication; };
+
 	/** @return Is this system configured to replicate object properties. */
 	bool AllowObjectReplication() { return bAllowObjectReplication; }
 
@@ -775,6 +788,9 @@ private:
 	int32 PIEInstanceID;
 	uint32 bIsServer : 1;
 	uint32 bAllowObjectReplication : 1;
+	uint32 bAllowExecuteRPC : 1;
+	uint32 bAllowExecuteRepNotify : 1;
+	uint32 bAlwaysSwapRolesOnReplication : 1;
 	uint32 bDoCollectGarbage : 1;
 };
 
