Support for RPC routing on the MultiServer Proxy when using Iris.

diff --git a/Engine/Plugins/Runtime/MultiServerReplication/Source/MultiServerReplication/Private/MultiServerProxy.cpp b/Engine/Plugins/Runtime/MultiServerReplication/Source/MultiServerReplication/Private/MultiServerProxy.cpp
--- a/Engine/Plugins/Runtime/MultiServerReplication/Source/MultiServerReplication/Private/MultiServerProxy.cpp
+++ b/Engine/Plugins/Runtime/MultiServerReplication/Source/MultiServerReplication/Private/MultiServerProxy.cpp
@@ -2,6 +2,7 @@
 
 #include "MultiServerProxy.h"
 #include "UnrealEngine.h"
+#include "EngineUtils.h"
 #include "Engine/LocalPlayer.h"
 #include "Engine/GameInstance.h"
 #include "Engine/NetworkObjectList.h"
@@ -572,97 +573,25 @@ void UProxyBackendNetDriver::ForwardRemoteFunction(UObject* RootObject, UObject*
 
 	UE_LOG(LogNetProxy, VeryVerbose, TEXT("RPC: Pre-forwarding %s on %s from game server %s to proxy connection."), *Function->GetName(), *OwningActor->GetName(), *GetName());
 
-	// The OwningActor actor will be associated with a connection (the return value of AActor::GetNetConnection()) to the backend game server, but
-	// when UProxyNetDriver::InternalProcessRemoteFunction() is called the owning player will be used to lookup the actual proxy connection to
-	// forward the RPC.
-	ProxyNetDriver->ProcessRemoteFunction(OwningActor, Function, Parms, static_cast<FOutParmRec*>(nullptr), static_cast<FFrame*>(nullptr), SubObject);
-}
-
-bool UProxyBackendNetDriver::ShouldSkipRepNotifies() const
-{
-	return true;
-}
-
-void UProxyBackendNetDriver::InternalProcessRemoteFunction(
-	class AActor* Actor,
-	class UObject* SubObject,
-	class UNetConnection* Connection,
-	class UFunction* Function,
-	void* Parms,
-	FOutParmRec* OutParms,
-	FFrame* Stack,
-	bool bIsServer)
-{
-	bool bShouldForwardRPC = true;
-
-	// If the sub-object is not owned by the actor, attempt to find an component in that actor that matches the same type.
-	// This logic assumes that an actor only has one component of a given type and will fail if that assumption is incorrect.
-	if (UActorComponent* SubObjectAsActorComponent = Cast<UActorComponent>(SubObject))
+	if (UReplicationSystem* ProxyReplicationSystem = ProxyNetDriver->GetReplicationSystem())
 	{
-		if (SubObjectAsActorComponent->GetOwner() != Actor)
-		{
-			int32 MatchingComponents = 0;
-			UActorComponent* MappedActorComponent = SubObjectAsActorComponent;
-
-			for (UActorComponent* ActorComponent : Actor->GetComponents())
-			{
-				if (ActorComponent->GetClass() == SubObjectAsActorComponent->GetClass())
-				{
-					MappedActorComponent = ActorComponent;
-					MatchingComponents++;
-
-					UE_LOG(LogNetProxy, VeryVerbose, TEXT("RPC: Remapping sub-object %s to sub-object %s in actor %s"),
-						*SubObjectAsActorComponent->GetName(),
-						*MappedActorComponent ->GetName(),
-						*Actor->GetName());
-				}
-			}
-
-			if (MappedActorComponent->GetOwner() != Actor)
-			{
-				bShouldForwardRPC = false;
-
-				UE_LOG(LogNetProxy, Verbose, TEXT("RPC: Unable map sub-object %s to actor %s"),
-					   *SubObjectAsActorComponent->GetName(),
-					   *Actor->GetName());
-			}
-
-			// Detect an actor with two components of the same type.
-			else if (MatchingComponents > 1)
-			{
-				bShouldForwardRPC = false;
-
-				UE_LOG(LogNetProxy, Verbose, TEXT("RPC: Found an actor %s with more than one component %s."), *Actor->GetName(), *SubObject->GetName());
-			}
-			else
-			{
-				SubObject = MappedActorComponent;
-			}
-		}
-	}
-
-	if (bShouldForwardRPC)
-	{
-		UE_LOG(LogNetProxy, Verbose, TEXT("RPC: Forwarding %s on %s (owner:%s sub-object:%s) to game server connection %s:%s"), 
-			   *Function->GetName(), 
-			   *Actor->GetName(), 
-			   *GetNameSafe(Actor->GetOwner()),
-			   *GetNameSafe(SubObject),
-			   *Connection->GetDriver()->GetName(),
-			   *Connection->GetName());
-
-		Super::InternalProcessRemoteFunction(Actor, SubObject, Connection, Function, Parms, OutParms, Stack, bIsServer);
+		ProxyNetDriver->RedirectRPCToProxyConnection(OwningActor, Function, [ProxyReplicationSystem, OwningActor, SubObject, Function, Parms](UNetConnection* ProxyConnection) {
+			ProxyReplicationSystem->SendRPC(ProxyConnection->GetConnectionHandle().GetParentConnectionId(), OwningActor, SubObject, Function, Parms);
+		});
 	}
 	else
 	{
-		UE_LOG(LogNetProxy, Verbose, TEXT("RPC: Ignoring %s on %s (sub-object:%s) to game server connection %s."), 
-			   *Function->GetName(), 
-			   *Actor->GetName(), 
-			   *GetNameSafe(SubObject),
-			   *Connection->GetName());
+		ProxyNetDriver->RedirectRPCToProxyConnection(OwningActor, Function, [this, OwningActor, SubObject, Function, Parms](UNetConnection* ProxyConnection) {
+			ProxyNetDriver->InternalProcessRemoteFunction(OwningActor, SubObject, ProxyConnection, Function, Parms, nullptr, nullptr, true);
+		});
 	}
 }
 
+bool UProxyBackendNetDriver::ShouldSkipRepNotifies() const
+{
+	return true;
+}
+
 bool UProxyBackendNetDriver::ShouldClientDestroyActor(AActor* Actor, EChannelCloseReason CloseReason) const
 {
 	// If an actor is destroyed remotely on a game server because it's being migrated to another server then it shouldn't be
@@ -691,6 +620,11 @@ void UProxyNetDriver::RegisterGameServer(const FURL& GameServerURL)
 	GameServerConnections.Emplace(GameServerURL);
 }
 
+void UProxyNetDriver::SetCyclePrimaryGameServer(bool bEnable)
+{
+	bCyclePrimaryGameServer = bEnable;
+}
+
 bool UProxyNetDriver::IsConnectedToAllGameServers() const
 {
 	for (const FGameServerConnectionState& GameServerConnectionState : GameServerConnections)
@@ -898,6 +832,54 @@ void UProxyNetDriver::ForwardRemoteFunction(UObject* RootObject, UObject* SubObj
 
 	UE_LOG(LogNetProxy, VeryVerbose, TEXT("RPC: Pre-forwarding %s on %s from proxy connection to game server."), *Function->GetName(), *OwningActor->GetName());
 
+	// If the sub-object is not owned by the actor, attempt to find an component in that actor that matches the same type.
+	// This logic assumes that an actor only has one component of a given type and will fail if that assumption is incorrect.
+	bool bShouldForwardRPC = true;
+
+	if (UActorComponent* SubObjectAsActorComponent = Cast<UActorComponent>(SubObject))
+	{
+		if (SubObjectAsActorComponent->GetOwner() != OwningActor)
+		{
+			int32 MatchingComponents = 0;
+			UActorComponent* MappedActorComponent = SubObjectAsActorComponent;
+
+			for (UActorComponent* ActorComponent : OwningActor->GetComponents())
+			{
+				if (ActorComponent->GetClass() == SubObjectAsActorComponent->GetClass())
+				{
+					MappedActorComponent = ActorComponent;
+					MatchingComponents++;
+
+					UE_LOG(LogNetProxy, VeryVerbose, TEXT("RPC: Remapping sub-object %s to sub-object %s in actor %s"),
+						   *SubObjectAsActorComponent->GetName(),
+						   *MappedActorComponent ->GetName(),
+						   *OwningActor->GetName());
+				}
+			}
+
+			if (MappedActorComponent->GetOwner() != OwningActor)
+			{
+				bShouldForwardRPC = false;
+
+				UE_LOG(LogNetProxy, Verbose, TEXT("RPC: Unable map sub-object %s to actor %s"),
+					   *SubObjectAsActorComponent->GetName(),
+					   *OwningActor->GetName());
+			}
+
+			// Detect an actor with two components of the same type.
+			else if (MatchingComponents > 1)
+			{
+				bShouldForwardRPC = false;
+
+				UE_LOG(LogNetProxy, Verbose, TEXT("RPC: Found an actor %s with more than one component %s."), *OwningActor->GetName(), *SubObject->GetName());
+			}
+			else
+			{
+				SubObject = MappedActorComponent;
+			}
+		}
+	}
+
 	if (TObjectPtr<const APlayerController>* GamePlayerControllerPtr = PlayerToPrimaryGameServerPlayerController.Find(OwningPlayer))
 	{
 		UNetConnection* GameConnection = GamePlayerControllerPtr->Get()->NetConnection;
@@ -906,7 +888,19 @@ void UProxyNetDriver::ForwardRemoteFunction(UObject* RootObject, UObject* SubObj
 			const bool bIsConnectionToGameServer = GameConnection->IsA(UProxyBackendNetConnection::StaticClass()) || GameConnection->IsA(UProxyBackendChildNetConnection::StaticClass());
 			if (ensure(bIsConnectionToGameServer))
 			{
-				GameConnection->Driver->ProcessRemoteFunction(OwningActor, Function, Parms, static_cast<FOutParmRec*>(nullptr), static_cast<FFrame*>(nullptr), SubObject);
+				if (bShouldForwardRPC)
+				{
+					UE_LOG(LogNetProxy, Verbose, TEXT("RPC: Forwarding %s on %s to game server connection %s:%s for player %s"), *GetNameSafe(Function), *GetNameSafe(OwningActor), *GameConnection->Driver->GetName(), *GetNameSafe(GameConnection), *GetNameSafe(OwningPlayer));
+					GameConnection->Driver->ProcessRemoteFunction(OwningActor, Function, Parms, static_cast<FOutParmRec*>(nullptr), static_cast<FFrame*>(nullptr), SubObject);
+				}
+				else
+				{
+					UE_LOG(LogNetProxy, Verbose, TEXT("RPC: Ignoring %s on %s (sub-object:%s) to game server connection %s."),
+						   *Function->GetName(),
+						   *OwningActor->GetName(),
+						   *GetNameSafe(SubObject),
+						   *GameConnection->GetName());
+				}
 			}
 			else
 			{
@@ -983,37 +977,6 @@ bool UProxyNetDriver::ShouldCallRemoteFunction(UObject* Object, UFunction* Funct
 	return !RepFlags.bIgnoreRPCs;
 }
 
-void UProxyNetDriver::InternalProcessRemoteFunction(
-	class AActor* Actor,
-	class UObject* SubObject,
-	class UNetConnection* Connection,
-	class UFunction* Function,
-	void* Parms,
-	FOutParmRec* OutParms,
-	FFrame* Stack,
-	bool bIsServer)
-{
-	if (ULocalPlayer* Player = Cast<ULocalPlayer>(Actor->GetNetOwningPlayerAnyRole()))
-	{
-		// RPCs from game servers will be routed to proxy connections. Only RPCs from a connections primary game server will
-		// be routed to the game client; other RPCs will be ignored.
-		if (UNetConnection* ProxyConnection = GetProxyConnectionFromPrimaryPlayer(Player))
-		{
-			UPlayer* Owner = Actor->GetNetOwningPlayerAnyRole();
-			UE_LOG(LogNetProxy, Verbose, TEXT("RPC: Forwarding %s on %s (owner:%s) to proxy connection %s:%s for player %s"), *GetNameSafe(Function), *GetNameSafe(Actor), *GetNameSafe(Owner), *ProxyConnection->Driver->GetName(), *GetNameSafe(ProxyConnection), *GetNameSafe(Player));
-			Super::InternalProcessRemoteFunction(Actor, SubObject, ProxyConnection, Function, Parms, OutParms, Stack, bIsServer);
-		}
-		else
-		{
-			UE_LOG(LogNetProxy, Verbose, TEXT("RPC: Ignoring %s on %s because player %s isn't a primary player."), *Function->GetName(), *Actor->GetName(), *Player->GetName());
-		}
-	}
-	else
-	{
-		UE_LOG(LogNetProxy, Verbose, TEXT("RPC: Ignoring %s on %s because there is no actor owning player."), *Function->GetName(), *Actor->GetName());
-	}
-}
-
 int32 UProxyNetDriver::ServerReplicateActors(float DeltaSeconds)
 {
 	PrepareStateForRelevancy();
@@ -1059,19 +1022,43 @@ void UProxyNetDriver::Shutdown()
 
 void UProxyNetDriver::TickFlush(float DeltaSeconds)
 {
-	// The function AssociatePlayerControllerWithProxyConnection() would ideally be called after a connection to a backend server
-	// has been established in UProxyBackendChildNetConnection::HandleClientPlayer(). This is not currently possible because 
-	// the actor is only added to the Iris replication system after HandleClientPlayer() is called (via OnActorChannelOpen()).
 	if (GetReplicationSystem())
 	{
-		for (TPair<TObjectPtr<UNetConnection>, TObjectPtr<const APlayerController>>& ProxyConnectionToPlayerController : ProxyConnectionToPrimaryGameServerPlayerController)
+		// Enumerating through all of the actors in the world is not efficient, but it's currently the most
+		// reliable way to ensure that all actor's owners are correctly set on the proxy.
+		// 
+		// Ideally the call to SetProxyConnectionAsOwner() should be made when an actor is registered with the 
+		// replication system, but there isn't currently a hook in the proxy driver to do that: AddNetworkActor() 
+		// and NotifyActorChannelOpen() cannot be used because they are called before the actor is registered 
+		// in the replication system.
+		for (TActorIterator<AActor> It(World); It; ++It)
 		{
-			TObjectPtr<UNetConnection> ProxyConnection = ProxyConnectionToPlayerController.Key;
-			TObjectPtr<const APlayerController> PlayerController = ProxyConnectionToPlayerController.Value;
+			const AActor* Actor = *It;
+			const AActor* Owner = Actor->GetOwner();
+			
+			if (!Actor->GetIsReplicated() || !(Actor->IsActorBeginningPlay() || Actor->HasActorBegunPlay()))
+			{
+				continue;
+			}
 
-			if (PlayerController->HasActorBegunPlay() || PlayerController->IsActorBeginningPlay())
+			if (Actor->GetOwner() == nullptr && !Actor->IsA(APlayerController::StaticClass()))
 			{
-				AssociatePlayerControllerWithProxyConnection(ProxyConnection, const_cast<APlayerController*>(PlayerController.Get()));
+				continue;
+			}
+
+			UNetConnection* ProxyConnection = nullptr;
+			for (TPair<TObjectPtr<UNetConnection>, TObjectPtr<const APlayerController>>& Pair : ProxyConnectionToPrimaryGameServerPlayerController)
+			{
+				if (Actor == Pair.Value || Owner == Pair.Value)
+				{
+					ProxyConnection = Pair.Key;
+					break;
+				}
+			}
+
+			if (ProxyConnection)
+			{
+				SetProxyConnectionAsOwner(ProxyConnection, Actor);
 			}
 		}
 	}
@@ -1458,18 +1445,37 @@ void UProxyNetDriver::ReceivedReassignedNoPawnPlayerController(const FMultiServe
 	}
 }
 
-void UProxyNetDriver::AssociatePlayerControllerWithProxyConnection(UNetConnection* ProxyConnection, APlayerController* PlayerController)
+void UProxyNetDriver::RedirectRPCToProxyConnection(AActor* Actor, UFunction* Function, TFunction<void(UNetConnection*)> RPCSendingFunction)
 {
-	// If these ensures hit then the player controller won't be added to the replication system leading to a connectivity issues.
-	ensure(PlayerController->GetIsReplicated());
-	ensure(PlayerController->IsActorBeginningPlay() || PlayerController->HasActorBegunPlay());
+	if (ULocalPlayer* Player = Cast<ULocalPlayer>(Actor->GetNetOwningPlayerAnyRole()))
+	{
+		// RPCs from game servers will be routed to proxy connections. Only RPCs from a connections primary game server will
+		// be routed to the game client; other RPCs will be ignored.
+		if (UNetConnection* ProxyConnection = GetProxyConnectionFromPrimaryPlayer(Player))
+		{
+			UPlayer* Owner = Actor->GetNetOwningPlayerAnyRole();
+			UE_LOG(LogNetProxy, Verbose, TEXT("RPC: Forwarding %s on %s (owner:%s) to proxy connection %s:%s for player %s"), *GetNameSafe(Function), *GetNameSafe(Actor), *GetNameSafe(Owner), *ProxyConnection->Driver->GetName(), *GetNameSafe(ProxyConnection), *GetNameSafe(Player));
+			RPCSendingFunction(ProxyConnection);
+		}
+		else
+		{
+			UE_LOG(LogNetProxy, Verbose, TEXT("RPC: Ignoring %s on %s because player %s isn't a primary player."), *Function->GetName(), *Actor->GetName(), *Player->GetName());
+		}
+	}
+	else
+	{
+		UE_LOG(LogNetProxy, Verbose, TEXT("RPC: Ignoring %s on %s because there is no actor owning player."), *Function->GetName(), *Actor->GetName());
+	}
+}
 
+void UProxyNetDriver::SetProxyConnectionAsOwner(UNetConnection* ProxyConnection, const AActor* Actor)
+{
 	if (UReplicationSystem* DriverReplicationSystem = GetReplicationSystem())
 	{
 		if (UObjectReplicationBridge* ObjectReplicationBridge = DriverReplicationSystem->GetReplicationBridgeAs<UObjectReplicationBridge>())
 		{
 			const uint32 ConnectionId = ProxyConnection->GetConnectionHandle().GetParentConnectionId();
-			const UE::Net::FNetRefHandle RefHandle = ObjectReplicationBridge->GetReplicatedRefHandle(PlayerController);
+			const UE::Net::FNetRefHandle RefHandle = ObjectReplicationBridge->GetReplicatedRefHandle(Actor);
 			if (RefHandle.IsValid())
 			{
 				DriverReplicationSystem->SetReplicationEnabledForConnection(ProxyConnection->GetConnectionHandle().GetParentConnectionId(), true);
diff --git a/Engine/Plugins/Runtime/MultiServerReplication/Source/MultiServerReplication/Public/MultiServerProxy.h b/Engine/Plugins/Runtime/MultiServerReplication/Source/MultiServerReplication/Public/MultiServerProxy.h
--- a/Engine/Plugins/Runtime/MultiServerReplication/Source/MultiServerReplication/Public/MultiServerProxy.h
+++ b/Engine/Plugins/Runtime/MultiServerReplication/Source/MultiServerReplication/Public/MultiServerProxy.h
@@ -281,7 +281,6 @@ public:
 	virtual bool InitBase(bool bInitAsClient, FNetworkNotify* InNotify, const FURL& URL, bool bReuseAddressAndPort, FString& Error);
 	virtual void ForwardRemoteFunction(UObject* RootObject, UObject* SubObject, UFunction* Function, void* Parms) override;
 	virtual bool ShouldSkipRepNotifies() const override;
-	virtual void InternalProcessRemoteFunction(class AActor* Actor, class UObject* SubObject, class UNetConnection* Connection, class UFunction* Function, void* Parms, FOutParmRec* OutParms, FFrame* Stack, bool bIsServer) override;
 	virtual bool ShouldClientDestroyActor(AActor* Actor, EChannelCloseReason CloseReason) const;
 
 	virtual bool ShouldUpdateStats() const override;
@@ -330,6 +329,9 @@ public:
 	/** Register a game server. */
 	void RegisterGameServer(const FURL& GameServerURL);
 
+	/** Cycle through game servers as primary game servers each time a client connects to the proxy. */
+	void SetCyclePrimaryGameServer(bool bEnable);
+
 	/** Return true if all registered servers are connected. */
 	bool IsConnectedToAllGameServers() const;
 
@@ -349,7 +351,6 @@ public:
 	virtual void AddNetworkActor(AActor* Actor) override;
 	virtual void RemoveNetworkActor(AActor* Actor) override;
 	virtual bool ShouldCallRemoteFunction(UObject* Object, UFunction* Function, const FReplicationFlags& RepFlags) const;
-	virtual void InternalProcessRemoteFunction(class AActor* Actor, class UObject* SubObject, class UNetConnection* Connection, class UFunction* Function, void* Parms, FOutParmRec* OutParms, FFrame* Stack, bool bIsServer) override;
 	virtual int32 ServerReplicateActors(float DeltaSeconds) override;
 	virtual bool CanDowngradeActorRole(UNetConnection* ProxyConnection, AActor* Actor) const override;
 	virtual void Shutdown() override;
@@ -405,12 +406,15 @@ private:
 	void ResetProxyNetworkMetrics();
 
 	/** 
-	 * Perform logic that associates a player controller with an incoming proxy connection to enable correct replication to that connection.
+	 * Send an RPC correct proxy connection based on the actor's owner and local player.
 	 *
-	 * This logic is required to glue the player controller used to connect to a backend game server with an incoming proxy connection and
-	 * only applies to the Iris replication system.
+	 * This function won't perform logic to send the RPC; that logic is expected to be called as part of the function
+	 * argument RPCSendingFunction. RPCSendingFunction will be provided the derived proxy connection, if one can be found, as an argument.
 	 */
-	void AssociatePlayerControllerWithProxyConnection(UNetConnection* ProxyConnection, APlayerController* PlayerController);
+	void RedirectRPCToProxyConnection(AActor* Actor, UFunction* Function, TFunction<void(UNetConnection*)> RPCSendingFunction);
+
+	/** (Iris only) Configure the connection to own this particular actor (required for replication conditionals and internal filtering logic). */
+	void SetProxyConnectionAsOwner(UNetConnection* ProxyConnection, const AActor* Actor);
 
 	/** 
 	 * Receive a new game player controller for an already open route.
diff --git a/Engine/Source/Runtime/Net/Iris/Private/Iris/ReplicationSystem/NetBlob/NetRPC.cpp b/Engine/Source/Runtime/Net/Iris/Private/Iris/ReplicationSystem/NetBlob/NetRPC.cpp
--- a/Engine/Source/Runtime/Net/Iris/Private/Iris/ReplicationSystem/NetBlob/NetRPC.cpp
+++ b/Engine/Source/Runtime/Net/Iris/Private/Iris/ReplicationSystem/NetBlob/NetRPC.cpp
@@ -625,6 +625,7 @@ void FNetRPC::CallFunction(FNetRPCCallContext& CallContext)
 	}
 
 	// Call function
+	if (ReplicationSystem->AllowExecuteRPC())
 	{
 #if IRIS_CLIENT_PROFILER_ENABLE
 		UE::Net::FClientProfiler::RecordRPC(ActualFunction->GetFName());
diff --git a/Engine/Source/Runtime/Net/Iris/Private/Iris/ReplicationSystem/NetBlob/NetRPCHandler.cpp b/Engine/Source/Runtime/Net/Iris/Private/Iris/ReplicationSystem/NetBlob/NetRPCHandler.cpp
--- a/Engine/Source/Runtime/Net/Iris/Private/Iris/ReplicationSystem/NetBlob/NetRPCHandler.cpp
+++ b/Engine/Source/Runtime/Net/Iris/Private/Iris/ReplicationSystem/NetBlob/NetRPCHandler.cpp
@@ -45,11 +45,6 @@ TRefCountPtr<UE::Net::FNetBlob> UNetRPCHandler::CreateNetBlob(const FNetBlobCrea
 
 void UNetRPCHandler::OnNetBlobReceived(UE::Net::FNetSerializationContext& NetContext, const TRefCountPtr<FNetBlob>& NetBlob)
 {
-	if (!NetContext.GetInternalContext()->ReplicationSystem->AllowExecuteRPC())
-	{
-		return;
-	}
-
 	const UE::Net::FForwardNetRPCCallMulticastDelegate& ForwardNetRPCCallDelegate = ReplicationSystem->GetReplicationSystemInternal()->GetForwardNetRPCCallMulticastDelegate();
 	UE::Net::FNetRPCCallContext CallContext(NetContext, ForwardNetRPCCallDelegate);
 	FNetRPC* RPC = static_cast<FNetRPC*>(NetBlob.GetReference());
