Fixes for Iris filtering in the MultiServerProxy:
- Make sure client's view targets are set in the Iris path.
- Make sure pawns have their location updated in the Iris path.

#FailedValidations Preflight

diff --git a/Engine/Plugins/Runtime/MultiServerReplication/Source/MultiServerReplication/Private/MultiServerProxy.cpp b/Engine/Plugins/Runtime/MultiServerReplication/Source/MultiServerReplication/Private/MultiServerProxy.cpp
--- a/Engine/Plugins/Runtime/MultiServerReplication/Source/MultiServerReplication/Private/MultiServerProxy.cpp
+++ b/Engine/Plugins/Runtime/MultiServerReplication/Source/MultiServerReplication/Private/MultiServerProxy.cpp
@@ -979,8 +979,6 @@ bool UProxyNetDriver::ShouldCallRemoteFunction(UObject* Object, UFunction* Funct
 
 int32 UProxyNetDriver::ServerReplicateActors(float DeltaSeconds)
 {
-	PrepareStateForRelevancy();
-
 	return Super::ServerReplicateActors(DeltaSeconds);
 }
 
@@ -1066,6 +1064,8 @@ void UProxyNetDriver::TickFlush(float DeltaSeconds)
 	// Update any metrics in the database that aren't tracked while ticking the frame.
 	UpdateProxyNetworkMetrics();
 
+	PrepareStateForRelevancy();
+
 	// The network stats and network metrics database are updated and reported to listeners here.
 	Super::TickFlush(DeltaSeconds);
 
@@ -1643,13 +1643,9 @@ void UProxyNetDriver::PrepareStateForRelevancy()
 
 	// Relevancy calcluations are made using the position returned by AActor::GetActorLocation(). Since the proxy is not simulating
 	// player movement the actor's location must be manually set to the location in AActor::ReplicatedMovement.
-	for (const TSharedPtr<FNetworkObjectInfo>& ObjectInfo : GetNetworkObjectList().GetAllObjects())
+	for (TActorIterator<APawn> It(World); It; ++It)
 	{
-		APawn* Pawn = Cast<APawn>(ObjectInfo->Actor);
-		if (!Pawn)
-		{
-			continue;
-		}
+		APawn* Pawn = *It;
 
 		Pawn->SetActorLocation(Pawn->GetReplicatedMovement().Location);
 	}
