Add some tests for filtering/relevancy on the MultiServer proxy.
MultiServer proxies now set the location of all actors replicating movement, not just pawns, in PrepareStateForRelevancy so that relevancy is calculated correctly for them.

diff --git a/Engine/Plugins/Runtime/MultiServerReplication/Source/MultiServerReplication/Private/MultiServerProxy.cpp b/Engine/Plugins/Runtime/MultiServerReplication/Source/MultiServerReplication/Private/MultiServerProxy.cpp
--- a/Engine/Plugins/Runtime/MultiServerReplication/Source/MultiServerReplication/Private/MultiServerProxy.cpp
+++ b/Engine/Plugins/Runtime/MultiServerReplication/Source/MultiServerReplication/Private/MultiServerProxy.cpp
@@ -1642,12 +1642,15 @@ void UProxyNetDriver::PrepareStateForRelevancy()
 	}
 
 	// Relevancy calcluations are made using the position returned by AActor::GetActorLocation(). Since the proxy is not simulating
-	// player movement the actor's location must be manually set to the location in AActor::ReplicatedMovement.
-	for (TActorIterator<APawn> It(World); It; ++It)
+	// actor movement the actor's location must be manually set to the location in AActor::ReplicatedMovement.
+	for (TActorIterator<AActor> It(World); It; ++It)
 	{
-		APawn* Pawn = *It;
+		AActor* Actor = *It;
 
-		Pawn->SetActorLocation(Pawn->GetReplicatedMovement().Location);
+		if (Actor->GetIsReplicated() && Actor->IsReplicatingMovement())
+		{
+			Actor->SetActorLocation(Actor->GetReplicatedMovement().Location);
+		}
 	}
 }
 
diff --git a/Engine/Source/Runtime/Engine/Private/Tests/NetTestHelpers.cpp b/Engine/Source/Runtime/Engine/Private/Tests/NetTestHelpers.cpp
--- a/Engine/Source/Runtime/Engine/Private/Tests/NetTestHelpers.cpp
+++ b/Engine/Source/Runtime/Engine/Private/Tests/NetTestHelpers.cpp
@@ -294,7 +294,7 @@ void FTestWorldInstance::Tick(float DeltaSeconds)
 	}
 }
 
-int32 FTestWorldInstance::GetPort()
+int32 FTestWorldInstance::GetPort() const
 {
 	if (const UWorld* const World = GetWorld())
 	{
diff --git a/Engine/Source/Runtime/Engine/Public/Tests/NetTestHelpers.h b/Engine/Source/Runtime/Engine/Public/Tests/NetTestHelpers.h
--- a/Engine/Source/Runtime/Engine/Public/Tests/NetTestHelpers.h
+++ b/Engine/Source/Runtime/Engine/Public/Tests/NetTestHelpers.h
@@ -140,7 +140,7 @@ public:
 	ENGINE_API UWorld* GetWorld() const;
 	ENGINE_API UNetDriver* GetNetDriver() const;
 
-	ENGINE_API int32 GetPort();
+	ENGINE_API int32 GetPort() const;
 
 	ENGINE_API void Tick(float DeltaSeconds = 0.0166f);
 
