Add NetDriver configs for the NetDrivers the MultiServerReplication plugin provides.

diff --git a/Engine/Plugins/Runtime/MultiServerReplication/Config/Engine.ini b/Engine/Plugins/Runtime/MultiServerReplication/Config/Engine.ini
new file mode 100644
--- /dev/null
+++ b/Engine/Plugins/Runtime/MultiServerReplication/Config/Engine.ini
@@ -0,0 +1,7 @@
+[/Script/Engine.Engine]
++NetDriverDefinitions=(DefName="MultiServerNetDriver",DriverClassName="/Script/MultiServerReplication.MultiServerNetDriver",DriverClassNameFallback="/Script/MultiServerReplication.MultiServerNetDriver")
++NetDriverDefinitions=(DefName="ProxyNetDriver",DriverClassName="/Script/MultiServerReplication.ProxyNetDriver",DriverClassNameFallback="/Script/MultiServerReplication.ProxyNetDriver")
++NetDriverDefinitions=(DefName="ProxyBackendNetDriver",DriverClassName="/Script/MultiServerReplication.ProxyBackendNetDriver",DriverClassNameFallback="/Script/MultiServerReplication.ProxyBackendNetDriver")
++IrisNetDriverConfigs=(NetDriverDefinition=MultiServerNetDriver, bCanUseIris=true, PreferredReplicationSystem=Iris)
++IrisNetDriverConfigs=(NetDriverDefinition=ProxyNetDriver, bCanUseIris=true)
++IrisNetDriverConfigs=(NetDriverDefinition=ProxyBackendNetDriver, bCanUseIris=true)
\ No newline at end of file
