MultiServerProxy:  Switching NetConnections in UProxyNetDriver::TickFlush so that Network & Game Code can correctly determine bNetOwner based on an Actor's GetNetConnection().

#rnx

diff --git a/Engine/Plugins/Runtime/MultiServerReplication/Source/MultiServerReplication/Private/MultiServerProxy.cpp b/Engine/Plugins/Runtime/MultiServerReplication/Source/MultiServerReplication/Private/MultiServerProxy.cpp
--- a/Engine/Plugins/Runtime/MultiServerReplication/Source/MultiServerReplication/Private/MultiServerProxy.cpp
+++ b/Engine/Plugins/Runtime/MultiServerReplication/Source/MultiServerReplication/Private/MultiServerProxy.cpp
@@ -1022,7 +1022,10 @@ void UProxyNetDriver::Shutdown()
 
 void UProxyNetDriver::TickFlush(float DeltaSeconds)
 {
-	if (GetReplicationSystem())
+	TArray<TWeakObjectPtr<APlayerController>> PlayerControllerNetConnectionsToRestore;
+
+	const bool bIris = (GetReplicationSystem() != nullptr);
+	if (bIris)
 	{
 		// Enumerating through all of the actors in the world is not efficient, but it's currently the most
 		// reliable way to ensure that all actor's owners are correctly set on the proxy.
@@ -1062,6 +1065,23 @@ void UProxyNetDriver::TickFlush(float DeltaSeconds)
 			}
 		}
 	}
+	else
+	{
+		// In the future: We can iterate over some kind of PlayerController->Connection map to make this more efficient.
+		// Reassign these PC NetConnections to the front-end (Client) Connections before replication
+		// Any code that checks the NetConnection for replication (bNetOwner) ownership will be correct.
+		for (TActorIterator<APlayerController> It(World); It; ++It)
+		{
+			APlayerController* PC = *It;
+			const FMultiServerProxyInternalConnectionRoute* Route = GetRouteWithPlayerController(PC);
+			if (Route)
+			{
+				ensureMsgf(PC->NetConnection == Route->GameServerConnection, TEXT("%s's NetConnection should have been pointing to the back-end GameServerConnection"), *PC->GetName());
+				PC->NetConnection = Route->ProxyConnection;
+				PlayerControllerNetConnectionsToRestore.Add(PC);
+			}
+		}
+	}
 
 	// Update any metrics in the database that aren't tracked while ticking the frame.
 	UpdateProxyNetworkMetrics();
@@ -1073,6 +1093,24 @@ void UProxyNetDriver::TickFlush(float DeltaSeconds)
 
 	// Reset any metrics that are tracked while ticking the frame.
 	ResetProxyNetworkMetrics();
+
+	if (!bIris)
+	{
+		// Restore the connection to the backend Game Servers
+		for (TWeakObjectPtr<APlayerController>& WeakPC : PlayerControllerNetConnectionsToRestore)
+		{
+			if (APlayerController* PC = WeakPC.Get())
+			{
+				const FMultiServerProxyInternalConnectionRoute* Route = GetRouteWithPlayerController(PC);
+				if (ensureMsgf(Route, TEXT("%s lost its MultiServerProxy Route during TickFlush"), *PC->GetName()))
+				{
+					ensureMsgf(PC->NetConnection == Route->ProxyConnection, TEXT("%s's NetConnection was (re-)created or (re-)assigned before TickFlush completed."), *PC->GetName());
+					PC->NetConnection = Route->GameServerConnection;
+				}
+			}
+		}
+
+	}
 }
 
 int32 UProxyNetDriver::GetNextGameServerClientId()
