MultiServerNode:  Adding const to the appropriate member functions.  Adding OnAllServersConnectedEvent as a potential Event to listen to.

#Synced-CL 48697860

diff --git a/Engine/Plugins/Runtime/MultiServerReplication/Source/MultiServerReplication/Private/MultiServerBeaconClient.cpp b/Engine/Plugins/Runtime/MultiServerReplication/Source/MultiServerReplication/Private/MultiServerBeaconClient.cpp
--- a/Engine/Plugins/Runtime/MultiServerReplication/Source/MultiServerReplication/Private/MultiServerBeaconClient.cpp
+++ b/Engine/Plugins/Runtime/MultiServerReplication/Source/MultiServerReplication/Private/MultiServerBeaconClient.cpp
@@ -122,7 +122,7 @@ void AMultiServerBeaconClient::ClientPeerConnected_Implementation(const FString&
 
 	if (OwningNode)
 	{
-		OwningNode->OnMultiServerConnected.ExecuteIfBound(OwningNode->GetLocalPeerId(), NewRemotePeerId, this);
+		OwningNode->NotifyMultiServerConnected(OwningNode->GetLocalPeerId(), NewRemotePeerId, this);
 	}
 }
 
@@ -204,7 +204,7 @@ void AMultiServerBeaconClient::ServerSetRemotePeerId_Implementation(const FStrin
 
 	if (OwningNode)
 	{
-		OwningNode->OnMultiServerConnected.ExecuteIfBound(OwningNode->GetLocalPeerId(), NewRemoteId, this);
+		OwningNode->NotifyMultiServerConnected(OwningNode->GetLocalPeerId(), NewRemoteId, this);
 	}
 }
 
diff --git a/Engine/Plugins/Runtime/MultiServerReplication/Source/MultiServerReplication/Private/MultiServerNode.cpp b/Engine/Plugins/Runtime/MultiServerReplication/Source/MultiServerReplication/Private/MultiServerNode.cpp
--- a/Engine/Plugins/Runtime/MultiServerReplication/Source/MultiServerReplication/Private/MultiServerNode.cpp
+++ b/Engine/Plugins/Runtime/MultiServerReplication/Source/MultiServerReplication/Private/MultiServerNode.cpp
@@ -12,6 +12,7 @@
 #include "Engine/NetDriver.h"
 #include "Misc/Parse.h"
 #include "Misc/CommandLine.h"
+#include "TimerManager.h"
 
 #include UE_INLINE_GENERATED_CPP_BY_NAME(MultiServerNode)
 
@@ -46,18 +47,8 @@ bool UMultiServerNode::AreAllServersConnected() const
 		return true;
 	}
 
-	// We are doing a bit more logic than "connected" and instead ensuring the Peer ID has been exchanged
-	uint32 NumAcknowledgedPeers = 0;
-	ForEachBeaconClient([&NumAcknowledgedPeers](AMultiServerBeaconClient* BeaconClient)
-	{
-		if (!BeaconClient->GetRemotePeerId().IsEmpty())
-		{
-			NumAcknowledgedPeers++;
-		}
-	});
-
-	const bool bAllPeersAreAcknowledged = NumAcknowledgedPeers >= (NumExpectedServers - 1);
-	return bAllPeersAreAcknowledged;
+	const bool bAllPeersAcknowledged = GetConnectionCount() >= GetNumExpectedConnections();
+	return bAllPeersAcknowledged;
 }
 
 UMultiServerNode* UMultiServerNode::Create(const FMultiServerNodeCreateParams& Params)
@@ -67,7 +58,6 @@ UMultiServerNode* UMultiServerNode::Create(const FMultiServerNodeCreateParams& P
 
 	if (bResult)
 	{
-		NewNode->NumExpectedServers = Params.NumServers;
 		NewNode->RegisterTickEvents();
 		return NewNode;
 	}
@@ -101,10 +91,10 @@ bool UMultiServerNode::RegisterServer(const FMultiServerNodeCreateParams& Params
 		return false;
 	}
 
+	NumExpectedServers = Params.NumServers;
 	LocalPeerId = Params.LocalPeerId;
 	UserBeaconClass = Params.UserBeaconClass;
-
-	OnMultiServerConnected = Params.OnMultiServerConnected;
+	OnMultiServerConnectedEvent = Params.OnMultiServerConnected;
 
 	if (Params.ListenPort == 0)
 	{
@@ -235,6 +225,14 @@ bool UMultiServerNode::RegisterServer(const FMultiServerNodeCreateParams& Params
 		}
 	}
 
+	if (NumExpectedServers <= 1)
+	{
+		Params.World->GetTimerManager().SetTimerForNextTick(FTimerDelegate::CreateWeakLambda(this, [this]()
+			{
+				OnAllServersConnectedEvent.Broadcast(this);
+			}));
+	}
+
 	return true;
 }
 
@@ -282,7 +280,7 @@ void UMultiServerNode::InternalTickFlush(float DeltaSeconds)
 		});
 }
 
-AMultiServerBeaconClient* UMultiServerNode::GetBeaconClientForRemotePeer(FStringView RemotePeerId)
+AMultiServerBeaconClient* UMultiServerNode::GetBeaconClientForRemotePeer(FStringView RemotePeerId) const
 {
 	// See if we are the host of the target server
 	UNetDriver* HostNetDriver = BeaconHost ? BeaconHost->GetNetDriver() : nullptr;
@@ -317,7 +315,7 @@ AMultiServerBeaconClient* UMultiServerNode::GetBeaconClientForRemotePeer(FString
 	return nullptr;
 }
 
-AMultiServerBeaconClient* UMultiServerNode::GetBeaconClientForURL(const FString& InURL)
+AMultiServerBeaconClient* UMultiServerNode::GetBeaconClientForURL(const FString& InURL) const
 {
 	FURL URL(nullptr, ToCStr(InURL), TRAVEL_Absolute);
 
@@ -383,7 +381,8 @@ void UMultiServerNode::ForEachNetDriver(TFunctionRef<void(UNetDriver*)> Operatio
 {
 	TSet<UNetDriver*> UniqueNetDrivers;
 
-	if (UNetDriver* HostNetDriver = BeaconHost->GetNetDriver())
+	UNetDriver* HostNetDriver = BeaconHost ? BeaconHost->GetNetDriver() : nullptr;
+	if (HostNetDriver)
 	{
 		UniqueNetDrivers.Add(HostNetDriver);
 	}
@@ -409,14 +408,28 @@ void UMultiServerNode::ForEachNetDriver(TFunctionRef<void(UNetDriver*)> Operatio
 	}
 }
 
-int32 UMultiServerNode::GetConnectionCount() const
+uint32 UMultiServerNode::GetConnectionCount() const
 {
-	int32 ConnectionCount = 0;
+	uint32 ConnectionCount = 0;
 
-	const_cast<UMultiServerNode*>(this)->ForEachBeaconClient([&ConnectionCount](AMultiServerBeaconClient*)
+	const_cast<UMultiServerNode*>(this)->ForEachBeaconClient([&ConnectionCount](AMultiServerBeaconClient* BeaconClient)
+	{
+		if (!BeaconClient->GetRemotePeerId().IsEmpty())
 		{
 			ConnectionCount++;
-		});
+		}
+	});
 
 	return ConnectionCount;
 }
+
+void UMultiServerNode::NotifyMultiServerConnected(const FString& LocalServerId, const FString& ConnectedRemoteServerId, AMultiServerBeaconClient* ConnectionToServer)
+{
+	OnMultiServerConnectedEvent.Broadcast(LocalServerId, ConnectedRemoteServerId, ConnectionToServer);
+
+	// If this was the last server to connect, now we're all connected
+	if (GetConnectionCount() == GetNumExpectedConnections())
+	{
+		OnAllServersConnectedEvent.Broadcast(this);
+	}
+}
diff --git a/Engine/Plugins/Runtime/MultiServerReplication/Source/MultiServerReplication/Public/MultiServerNode.h b/Engine/Plugins/Runtime/MultiServerReplication/Source/MultiServerReplication/Public/MultiServerNode.h
--- a/Engine/Plugins/Runtime/MultiServerReplication/Source/MultiServerReplication/Public/MultiServerNode.h
+++ b/Engine/Plugins/Runtime/MultiServerReplication/Source/MultiServerReplication/Public/MultiServerNode.h
@@ -16,7 +16,8 @@ class AMultiServerBeaconHostObject;
 class AMultiServerBeaconClient;
 class UMultiServerPeerConnection;
 
-DECLARE_DELEGATE_ThreeParams(FOnMultiServerConnected, const FString&, const FString&, AMultiServerBeaconClient*);
+DECLARE_MULTICAST_DELEGATE_ThreeParams(FOnMultiServerConnectedEvent, const FString&, const FString&, AMultiServerBeaconClient*);
+DECLARE_MULTICAST_DELEGATE_OneParam(FOnAllServersConnectedEvent, UMultiServerNode*);
 
 /* Parameters for initializing a UMultiServerNode */
 struct FMultiServerNodeCreateParams
@@ -34,7 +35,7 @@ struct FMultiServerNodeCreateParams
 	uint16 ListenPort = 0;
 
 	/** The number of servers we're expecting in the cluster */
-	uint32 NumServers = 0;
+	uint32 NumServers = 1;
 
 	/** List of addresses of other nodes to attempt to connect to. */
 	TArray<FString> PeerAddresses;
@@ -43,7 +44,7 @@ struct FMultiServerNodeCreateParams
 	TSubclassOf<AMultiServerBeaconClient> UserBeaconClass;
 
 	/** Callback invoked when a connection to a remote node is established. */
-	FOnMultiServerConnected OnMultiServerConnected;
+	FOnMultiServerConnectedEvent OnMultiServerConnected;
 };
 
 /**
@@ -79,15 +80,15 @@ public:
 
 	bool RegisterServer(const FMultiServerNodeCreateParams& Params);
 
-	AMultiServerBeaconClient* GetBeaconClientForRemotePeer(FStringView RemotePeerId);
+	AMultiServerBeaconClient* GetBeaconClientForRemotePeer(FStringView RemotePeerId) const;
 
 	template<class T>
-	T* GetBeaconClientForRemotePeer(FStringView RemotePeerId);
+	T* GetBeaconClientForRemotePeer(FStringView RemotePeerId) const;
 
-	AMultiServerBeaconClient* GetBeaconClientForURL(const FString& InURL);
+	AMultiServerBeaconClient* GetBeaconClientForURL(const FString& InURL) const;
 
 	template<class T>
-	T* GetBeaconClientForURL(const FString& InURL);
+	T* GetBeaconClientForURL(const FString& InURL) const;
 
 	FString GetLocalPeerId() const { return LocalPeerId; }
 
@@ -100,12 +101,25 @@ public:
 
 	void ForEachNetDriver(TFunctionRef<void(UNetDriver*)> Operation) const;
 
-	int32 GetConnectionCount() const;
+	/** How many acknolwedged connections do we have? */
+	uint32 GetConnectionCount() const;
+
+	/** How many connections are we expecting? */
+	uint32 GetNumExpectedConnections() const { return (NumExpectedServers > 0) ? NumExpectedServers - 1 : 0; }
+
+	/** Notify this MultiServerNode that a Beacon has connected to it */
+	void NotifyMultiServerConnected(const FString& LocalServerId, const FString& ConnectedRemoteServerId, AMultiServerBeaconClient* ConnectionToServer);
+
+	/** Allow users to register a delegate for when all servers are connected */
+	TMulticastDelegateRegistration<void(UMultiServerNode*)>& OnAllServersConnected()
+	{
+		return OnAllServersConnectedEvent;
+	}
 
 private:
-	friend AMultiServerBeaconClient;
+	FOnAllServersConnectedEvent OnAllServersConnectedEvent;
 
-	FOnMultiServerConnected OnMultiServerConnected;
+	FOnMultiServerConnectedEvent OnMultiServerConnectedEvent;
 
 	FString LocalPeerId;
 
@@ -126,7 +140,7 @@ private:
 
 	TSubclassOf<AMultiServerBeaconClient> UserBeaconClass;
 
-	// How many servers are we expecting? Used for AreAllServersConnected.
+	// How many servers are we expecting (including ourselves). This should be at least 1.
 	uint32 NumExpectedServers = 1;
 
 	// Multi-server nodes handle ticking of their NetDrivers directly instead of letting the world tick them.
@@ -148,14 +162,13 @@ private:
 };
 
 template<class T>
-T* UMultiServerNode::GetBeaconClientForRemotePeer(FStringView RemotePeerId)
+T* UMultiServerNode::GetBeaconClientForRemotePeer(FStringView RemotePeerId) const
 {
 	return Cast<T>(GetBeaconClientForRemotePeer(RemotePeerId));
 }
 
 template<class T>
-T* UMultiServerNode::GetBeaconClientForURL(const FString& InURL)
-
+T* UMultiServerNode::GetBeaconClientForURL(const FString& InURL) const
 {
 	return Cast<T>(GetBeaconClientForURL(InURL));
 }
\ No newline at end of file
