[Backout] - CL48766513
#rnx
Original CL Desc
-----------------------------------------------------------------
MultiServerNode:  Adding const to the appropriate member functions.  Adding OnAllServersConnectedEvent as a potential Event to listen to.
Backing out while I fix-up the EngineTest.

#Synced-CL 48697860

#Synced-CL 48766513

diff --git a/Engine/Plugins/Runtime/MultiServerReplication/Source/MultiServerReplication/Private/MultiServerBeaconClient.cpp b/Engine/Plugins/Runtime/MultiServerReplication/Source/MultiServerReplication/Private/MultiServerBeaconClient.cpp
--- a/Engine/Plugins/Runtime/MultiServerReplication/Source/MultiServerReplication/Private/MultiServerBeaconClient.cpp
+++ b/Engine/Plugins/Runtime/MultiServerReplication/Source/MultiServerReplication/Private/MultiServerBeaconClient.cpp
@@ -122,7 +122,7 @@ void AMultiServerBeaconClient::ClientPeerConnected_Implementation(const FString&
 
 	if (OwningNode)
 	{
-		OwningNode->NotifyMultiServerConnected(OwningNode->GetLocalPeerId(), NewRemotePeerId, this);
+		OwningNode->OnMultiServerConnected.ExecuteIfBound(OwningNode->GetLocalPeerId(), NewRemotePeerId, this);
 	}
 }
 
@@ -204,7 +204,7 @@ void AMultiServerBeaconClient::ServerSetRemotePeerId_Implementation(const FStrin
 
 	if (OwningNode)
 	{
-		OwningNode->NotifyMultiServerConnected(OwningNode->GetLocalPeerId(), NewRemoteId, this);
+		OwningNode->OnMultiServerConnected.ExecuteIfBound(OwningNode->GetLocalPeerId(), NewRemoteId, this);
 	}
 }
 
diff --git a/Engine/Plugins/Runtime/MultiServerReplication/Source/MultiServerReplication/Private/MultiServerNode.cpp b/Engine/Plugins/Runtime/MultiServerReplication/Source/MultiServerReplication/Private/MultiServerNode.cpp
--- a/Engine/Plugins/Runtime/MultiServerReplication/Source/MultiServerReplication/Private/MultiServerNode.cpp
+++ b/Engine/Plugins/Runtime/MultiServerReplication/Source/MultiServerReplication/Private/MultiServerNode.cpp
@@ -12,7 +12,6 @@
 #include "Engine/NetDriver.h"
 #include "Misc/Parse.h"
 #include "Misc/CommandLine.h"
-#include "TimerManager.h"
 
 #include UE_INLINE_GENERATED_CPP_BY_NAME(MultiServerNode)
 
@@ -47,8 +46,18 @@ bool UMultiServerNode::AreAllServersConnected() const
 		return true;
 	}
 
-	const bool bAllPeersAcknowledged = GetConnectionCount() >= GetNumExpectedConnections();
-	return bAllPeersAcknowledged;
+	// We are doing a bit more logic than "connected" and instead ensuring the Peer ID has been exchanged
+	uint32 NumAcknowledgedPeers = 0;
+	ForEachBeaconClient([&NumAcknowledgedPeers](AMultiServerBeaconClient* BeaconClient)
+	{
+		if (!BeaconClient->GetRemotePeerId().IsEmpty())
+		{
+			NumAcknowledgedPeers++;
+		}
+	});
+
+	const bool bAllPeersAreAcknowledged = NumAcknowledgedPeers >= (NumExpectedServers - 1);
+	return bAllPeersAreAcknowledged;
 }
 
 UMultiServerNode* UMultiServerNode::Create(const FMultiServerNodeCreateParams& Params)
@@ -58,6 +67,7 @@ UMultiServerNode* UMultiServerNode::Create(const FMultiServerNodeCreateParams& P
 
 	if (bResult)
 	{
+		NewNode->NumExpectedServers = Params.NumServers;
 		NewNode->RegisterTickEvents();
 		return NewNode;
 	}
@@ -91,10 +101,10 @@ bool UMultiServerNode::RegisterServer(const FMultiServerNodeCreateParams& Params
 		return false;
 	}
 
-	NumExpectedServers = Params.NumServers;
 	LocalPeerId = Params.LocalPeerId;
 	UserBeaconClass = Params.UserBeaconClass;
-	OnMultiServerConnectedEvent = Params.OnMultiServerConnected;
+
+	OnMultiServerConnected = Params.OnMultiServerConnected;
 
 	if (Params.ListenPort == 0)
 	{
@@ -225,14 +235,6 @@ bool UMultiServerNode::RegisterServer(const FMultiServerNodeCreateParams& Params
 		}
 	}
 
-	if (NumExpectedServers <= 1)
-	{
-		Params.World->GetTimerManager().SetTimerForNextTick(FTimerDelegate::CreateWeakLambda(this, [this]()
-			{
-				OnAllServersConnectedEvent.Broadcast(this);
-			}));
-	}
-
 	return true;
 }
 
@@ -280,7 +282,7 @@ void UMultiServerNode::InternalTickFlush(float DeltaSeconds)
 		});
 }
 
-AMultiServerBeaconClient* UMultiServerNode::GetBeaconClientForRemotePeer(FStringView RemotePeerId) const
+AMultiServerBeaconClient* UMultiServerNode::GetBeaconClientForRemotePeer(FStringView RemotePeerId)
 {
 	// See if we are the host of the target server
 	UNetDriver* HostNetDriver = BeaconHost ? BeaconHost->GetNetDriver() : nullptr;
@@ -315,7 +317,7 @@ AMultiServerBeaconClient* UMultiServerNode::GetBeaconClientForRemotePeer(FString
 	return nullptr;
 }
 
-AMultiServerBeaconClient* UMultiServerNode::GetBeaconClientForURL(const FString& InURL) const
+AMultiServerBeaconClient* UMultiServerNode::GetBeaconClientForURL(const FString& InURL)
 {
 	FURL URL(nullptr, ToCStr(InURL), TRAVEL_Absolute);
 
@@ -381,8 +383,7 @@ void UMultiServerNode::ForEachNetDriver(TFunctionRef<void(UNetDriver*)> Operatio
 {
 	TSet<UNetDriver*> UniqueNetDrivers;
 
-	UNetDriver* HostNetDriver = BeaconHost ? BeaconHost->GetNetDriver() : nullptr;
-	if (HostNetDriver)
+	if (UNetDriver* HostNetDriver = BeaconHost->GetNetDriver())
 	{
 		UniqueNetDrivers.Add(HostNetDriver);
 	}
@@ -408,28 +409,14 @@ void UMultiServerNode::ForEachNetDriver(TFunctionRef<void(UNetDriver*)> Operatio
 	}
 }
 
-uint32 UMultiServerNode::GetConnectionCount() const
+int32 UMultiServerNode::GetConnectionCount() const
 {
-	uint32 ConnectionCount = 0;
+	int32 ConnectionCount = 0;
 
-	const_cast<UMultiServerNode*>(this)->ForEachBeaconClient([&ConnectionCount](AMultiServerBeaconClient* BeaconClient)
-	{
-		if (!BeaconClient->GetRemotePeerId().IsEmpty())
+	const_cast<UMultiServerNode*>(this)->ForEachBeaconClient([&ConnectionCount](AMultiServerBeaconClient*)
 		{
 			ConnectionCount++;
-		}
-	});
+		});
 
 	return ConnectionCount;
 }
-
-void UMultiServerNode::NotifyMultiServerConnected(const FString& LocalServerId, const FString& ConnectedRemoteServerId, AMultiServerBeaconClient* ConnectionToServer)
-{
-	OnMultiServerConnectedEvent.Broadcast(LocalServerId, ConnectedRemoteServerId, ConnectionToServer);
-
-	// If this was the last server to connect, now we're all connected
-	if (GetConnectionCount() == GetNumExpectedConnections())
-	{
-		OnAllServersConnectedEvent.Broadcast(this);
-	}
-}
diff --git a/Engine/Plugins/Runtime/MultiServerReplication/Source/MultiServerReplication/Public/MultiServerNode.h b/Engine/Plugins/Runtime/MultiServerReplication/Source/MultiServerReplication/Public/MultiServerNode.h
--- a/Engine/Plugins/Runtime/MultiServerReplication/Source/MultiServerReplication/Public/MultiServerNode.h
+++ b/Engine/Plugins/Runtime/MultiServerReplication/Source/MultiServerReplication/Public/MultiServerNode.h
@@ -16,8 +16,7 @@ class AMultiServerBeaconHostObject;
 class AMultiServerBeaconClient;
 class UMultiServerPeerConnection;
 
-DECLARE_MULTICAST_DELEGATE_ThreeParams(FOnMultiServerConnectedEvent, const FString&, const FString&, AMultiServerBeaconClient*);
-DECLARE_MULTICAST_DELEGATE_OneParam(FOnAllServersConnectedEvent, UMultiServerNode*);
+DECLARE_DELEGATE_ThreeParams(FOnMultiServerConnected, const FString&, const FString&, AMultiServerBeaconClient*);
 
 /* Parameters for initializing a UMultiServerNode */
 struct FMultiServerNodeCreateParams
@@ -35,7 +34,7 @@ struct FMultiServerNodeCreateParams
 	uint16 ListenPort = 0;
 
 	/** The number of servers we're expecting in the cluster */
-	uint32 NumServers = 1;
+	uint32 NumServers = 0;
 
 	/** List of addresses of other nodes to attempt to connect to. */
 	TArray<FString> PeerAddresses;
@@ -44,7 +43,7 @@ struct FMultiServerNodeCreateParams
 	TSubclassOf<AMultiServerBeaconClient> UserBeaconClass;
 
 	/** Callback invoked when a connection to a remote node is established. */
-	FOnMultiServerConnectedEvent OnMultiServerConnected;
+	FOnMultiServerConnected OnMultiServerConnected;
 };
 
 /**
@@ -80,15 +79,15 @@ public:
 
 	bool RegisterServer(const FMultiServerNodeCreateParams& Params);
 
-	AMultiServerBeaconClient* GetBeaconClientForRemotePeer(FStringView RemotePeerId) const;
+	AMultiServerBeaconClient* GetBeaconClientForRemotePeer(FStringView RemotePeerId);
 
 	template<class T>
-	T* GetBeaconClientForRemotePeer(FStringView RemotePeerId) const;
+	T* GetBeaconClientForRemotePeer(FStringView RemotePeerId);
 
-	AMultiServerBeaconClient* GetBeaconClientForURL(const FString& InURL) const;
+	AMultiServerBeaconClient* GetBeaconClientForURL(const FString& InURL);
 
 	template<class T>
-	T* GetBeaconClientForURL(const FString& InURL) const;
+	T* GetBeaconClientForURL(const FString& InURL);
 
 	FString GetLocalPeerId() const { return LocalPeerId; }
 
@@ -101,25 +100,12 @@ public:
 
 	void ForEachNetDriver(TFunctionRef<void(UNetDriver*)> Operation) const;
 
-	/** How many acknolwedged connections do we have? */
-	uint32 GetConnectionCount() const;
-
-	/** How many connections are we expecting? */
-	uint32 GetNumExpectedConnections() const { return (NumExpectedServers > 0) ? NumExpectedServers - 1 : 0; }
-
-	/** Notify this MultiServerNode that a Beacon has connected to it */
-	void NotifyMultiServerConnected(const FString& LocalServerId, const FString& ConnectedRemoteServerId, AMultiServerBeaconClient* ConnectionToServer);
-
-	/** Allow users to register a delegate for when all servers are connected */
-	TMulticastDelegateRegistration<void(UMultiServerNode*)>& OnAllServersConnected()
-	{
-		return OnAllServersConnectedEvent;
-	}
+	int32 GetConnectionCount() const;
 
 private:
-	FOnAllServersConnectedEvent OnAllServersConnectedEvent;
+	friend AMultiServerBeaconClient;
 
-	FOnMultiServerConnectedEvent OnMultiServerConnectedEvent;
+	FOnMultiServerConnected OnMultiServerConnected;
 
 	FString LocalPeerId;
 
@@ -140,7 +126,7 @@ private:
 
 	TSubclassOf<AMultiServerBeaconClient> UserBeaconClass;
 
-	// How many servers are we expecting (including ourselves). This should be at least 1.
+	// How many servers are we expecting? Used for AreAllServersConnected.
 	uint32 NumExpectedServers = 1;
 
 	// Multi-server nodes handle ticking of their NetDrivers directly instead of letting the world tick them.
@@ -162,13 +148,14 @@ private:
 };
 
 template<class T>
-T* UMultiServerNode::GetBeaconClientForRemotePeer(FStringView RemotePeerId) const
+T* UMultiServerNode::GetBeaconClientForRemotePeer(FStringView RemotePeerId)
 {
 	return Cast<T>(GetBeaconClientForRemotePeer(RemotePeerId));
 }
 
 template<class T>
-T* UMultiServerNode::GetBeaconClientForURL(const FString& InURL) const
+T* UMultiServerNode::GetBeaconClientForURL(const FString& InURL)
+
 {
 	return Cast<T>(GetBeaconClientForURL(InURL));
 }
\ No newline at end of file
