Fix a log statement in UProxyNetDriver so that it correctly prints out the route's player.

#Synced-CL 49559917

diff --git a/Engine/Plugins/Runtime/MultiServerReplication/Source/MultiServerReplication/Private/MultiServerProxy.cpp b/Engine/Plugins/Runtime/MultiServerReplication/Source/MultiServerReplication/Private/MultiServerProxy.cpp
--- a/Engine/Plugins/Runtime/MultiServerReplication/Source/MultiServerReplication/Private/MultiServerProxy.cpp
+++ b/Engine/Plugins/Runtime/MultiServerReplication/Source/MultiServerReplication/Private/MultiServerProxy.cpp
@@ -1138,12 +1138,6 @@ void UProxyNetDriver::DisableActorLogicAndGameCode()
 
 void UProxyNetDriver::GameServerAssignPlayerController(UNetConnection* ChildGameServerConnection, UNetConnection* NotUsedParentGameServerConnection, APlayerController* GameServerPlayerController)
 {
-	UE_LOG(LogNetProxy, Log, TEXT("[%d] Received a new player controller %s:%s for connection %s:%s for player %s."), 
-		   GameServerPlayerController->GetClientHandshakeId(),
-		   *GameServerPlayerController->GetName(), *GameServerPlayerController->GetClass()->GetName(),
-		   *ChildGameServerConnection->GetDriver()->GetName(), *ChildGameServerConnection->GetName(),
-		   *ChildGameServerConnection->PlayerId.ToDebugString());
-
 	const uint32 ClientHandshakeId = GameServerPlayerController->GetClientHandshakeId();
 	FMultiServerProxyInternalConnectionRoute* Route = InternalRoutes.Find(ClientHandshakeId);
 	if (!ensure(Route))
@@ -1151,6 +1145,12 @@ void UProxyNetDriver::GameServerAssignPlayerController(UNetConnection* ChildGame
 		return;
 	}
 
+	UE_LOG(LogNetProxy, Log, TEXT("[%d] Received a new player controller %s:%s for connection %s:%s for player %s."), 
+		   GameServerPlayerController->GetClientHandshakeId(),
+		   *GameServerPlayerController->GetName(), *GameServerPlayerController->GetClass()->GetName(),
+		   *ChildGameServerConnection->GetDriver()->GetName(), *ChildGameServerConnection->GetName(),
+		   *GetNameSafe(Route->Player));
+
 	EProxyConnectionState State = Route->State;
 
 	switch (State)
