Fix a crash when a proxy client closes it's connection.

#Synced-CL 49844415

diff --git a/Engine/Plugins/Runtime/MultiServerReplication/Source/MultiServerReplication/Private/MultiServerProxy.cpp b/Engine/Plugins/Runtime/MultiServerReplication/Source/MultiServerReplication/Private/MultiServerProxy.cpp
--- a/Engine/Plugins/Runtime/MultiServerReplication/Source/MultiServerReplication/Private/MultiServerProxy.cpp
+++ b/Engine/Plugins/Runtime/MultiServerReplication/Source/MultiServerReplication/Private/MultiServerProxy.cpp
@@ -1215,7 +1215,11 @@ void UProxyNetDriver::HandleClosedProxyConnection(UNetConnection* ProxyConnectio
 		const bool bStillConnecting  = (Route->State == EProxyConnectionState::Connecting || Route->State == EProxyConnectionState::ConnectingPrimary);
 		if (!bStillConnecting)
 		{
-			CloseAndCleanUpConnectionRoute(Route);
+			FMultiServerProxyInternalConnectionRoute RemovedRoute;
+			if (InternalRoutes.RemoveAndCopyValue(ClientHandshakeIdToRemove, RemovedRoute))
+			{
+				CloseAndCleanUpConnectionRoute(&RemovedRoute);
+			}
 		}
 		else
 		{
@@ -1242,7 +1246,6 @@ void UProxyNetDriver::CloseAndCleanUpConnectionRoute(const FMultiServerProxyInte
 		   *UE::MultiServerProxy::Private::PlayerControllerToString(Route->PlayerController),
 		   *GetNameSafe(Route->Player));
 
-	InternalRoutes.Remove(Route->ClientHandshakeId);
 	UChildConnection* GameServerConnectionAsChild = Cast<UChildConnection>(Route->GameServerConnection);
 	if (ensure(GameServerConnectionAsChild))
 	{
@@ -1389,7 +1392,11 @@ void UProxyNetDriver::ProcessConnectionState_PendingClose(FMultiServerProxyInter
 	Route->GameServerConnection = GameServerConnection;
 	GameServerConnection->PlayerController = PlayerController;
 
-	CloseAndCleanUpConnectionRoute(Route);
+	FMultiServerProxyInternalConnectionRoute RemovedRoute;
+	if (InternalRoutes.RemoveAndCopyValue(Route->ClientHandshakeId, RemovedRoute))
+	{
+		CloseAndCleanUpConnectionRoute(&RemovedRoute);
+	}
 }
 
 void UProxyNetDriver::UpdateProxyNetworkMetrics()
