MultiServerReplication:  The AllServersConnected delegate should always fire outside of an AutoRTFM transaction.

#rnx

#Synced-CL 50269217

diff --git a/Engine/Plugins/Runtime/MultiServerReplication/Source/MultiServerReplication/Private/MultiServerNode.cpp b/Engine/Plugins/Runtime/MultiServerReplication/Source/MultiServerReplication/Private/MultiServerNode.cpp
--- a/Engine/Plugins/Runtime/MultiServerReplication/Source/MultiServerReplication/Private/MultiServerNode.cpp
+++ b/Engine/Plugins/Runtime/MultiServerReplication/Source/MultiServerReplication/Private/MultiServerNode.cpp
@@ -229,7 +229,10 @@ bool UMultiServerNode::RegisterServer(const FMultiServerNodeCreateParams& Params
 	{
 		Params.World->GetTimerManager().SetTimerForNextTick(FTimerDelegate::CreateWeakLambda(this, [this]()
 			{
-				OnAllServersConnectedEvent.Broadcast(this);
+				UE_AUTORTFM_ONCOMMIT(this)
+				{
+					BroadcastAllServersConnected();
+				};
 			}));
 	}
 
