Prevent the proxy server from calling into UGameInstance virtual functions when a player connects and disconnects.

#Synced-CL 51341769

diff --git a/Engine/Plugins/Runtime/MultiServerReplication/Source/MultiServerReplication/Private/MultiServerProxy.cpp b/Engine/Plugins/Runtime/MultiServerReplication/Source/MultiServerReplication/Private/MultiServerProxy.cpp
--- a/Engine/Plugins/Runtime/MultiServerReplication/Source/MultiServerReplication/Private/MultiServerProxy.cpp
+++ b/Engine/Plugins/Runtime/MultiServerReplication/Source/MultiServerReplication/Private/MultiServerProxy.cpp
@@ -326,7 +326,7 @@ void UProxyListenerNotify::ConnectToGameServer(UNetConnection* ProxyConnection,
 		const FPlatformUserId DummyClientId = FPlatformUserId::CreateFromInternalId(ProxyNetDriver->GetNextGameServerClientId());
 		ULocalPlayer* DummyPlayer = NewObject<ULocalPlayer>(GEngine, GEngine->LocalPlayerClass);
 		DummyPlayer->SetConnectionIdentifier(ParentClientHandshakeId);
-		ProxyNetDriver->GetWorld()->GetGameInstance()->AddLocalPlayer(DummyPlayer, DummyClientId);
+		ProxyNetDriver->GetWorld()->GetGameInstance()->UGameInstance::AddLocalPlayer(DummyPlayer, DummyClientId);
 
 		FURL URL = GameServerConnectionState->GameServerURL;
 		
@@ -369,7 +369,7 @@ void UProxyListenerNotify::ConnectToGameServer(UNetConnection* ProxyConnection,
 		const FPlatformUserId GameServerClientId = FPlatformUserId::CreateFromInternalId(ProxyNetDriver->GetNextGameServerClientId());
 		ULocalPlayer* NewPlayer = NewObject<ULocalPlayer>(GEngine, GEngine->LocalPlayerClass);
 		NewPlayer->SetConnectionIdentifier(ChildClientHandshakeId);
-		ProxyNetDriver->GetWorld()->GetGameInstance()->AddLocalPlayer(NewPlayer, GameServerClientId);
+		ProxyNetDriver->GetWorld()->GetGameInstance()->UGameInstance::AddLocalPlayer(NewPlayer, GameServerClientId);
 
 		// The new player will use the same unique identifier as the incoming proxy connection so that it will be propogated up to 
 		// the game servers through UNetConnection::PlayerId. This way each game server's incoming connection will have a PlayerId 
@@ -474,6 +474,11 @@ void UProxyBackendNetConnection::HandleClientPlayer(APlayerController* NewPlayer
 	}
 }
 
+void UProxyBackendChildNetConnection::RemoveLocalPlayer(UGameInstance* GameInstance, ULocalPlayer* LocalPlayer)
+{
+	GameInstance->UGameInstance::RemoveLocalPlayer(LocalPlayer);
+}
+
 void UProxyBackendChildNetConnection::HandleClientPlayer(APlayerController* NewPlayerController, UNetConnection* GameServerConnection)
 {
 	// This function is called when a PlayerController is replicated to the proxy from a game server and represents the finalization
diff --git a/Engine/Plugins/Runtime/MultiServerReplication/Source/MultiServerReplication/Public/MultiServerProxy.h b/Engine/Plugins/Runtime/MultiServerReplication/Source/MultiServerReplication/Public/MultiServerProxy.h
--- a/Engine/Plugins/Runtime/MultiServerReplication/Source/MultiServerReplication/Public/MultiServerProxy.h
+++ b/Engine/Plugins/Runtime/MultiServerReplication/Source/MultiServerReplication/Public/MultiServerProxy.h
@@ -265,6 +265,10 @@ class MULTISERVERREPLICATION_API UProxyBackendChildNetConnection : public UChild
 public:
 
 	virtual void HandleClientPlayer(APlayerController* PC, UNetConnection* NetConnection);
+	
+protected:
+
+	virtual void RemoveLocalPlayer(UGameInstance* GameInstance, ULocalPlayer* LocalPlayer) override;
 };
 
 /** 
diff --git a/Engine/Source/Runtime/Engine/Classes/Engine/ChildConnection.h b/Engine/Source/Runtime/Engine/Classes/Engine/ChildConnection.h
--- a/Engine/Source/Runtime/Engine/Classes/Engine/ChildConnection.h
+++ b/Engine/Source/Runtime/Engine/Classes/Engine/ChildConnection.h
@@ -110,4 +110,7 @@ public:
 
 private:
 	void AssignConnectionHandle();
+
+protected:
+	virtual void RemoveLocalPlayer(UGameInstance* GameInstance, ULocalPlayer* LocalPlayer);
 };
diff --git a/Engine/Source/Runtime/Engine/Private/NetConnection.cpp b/Engine/Source/Runtime/Engine/Private/NetConnection.cpp
--- a/Engine/Source/Runtime/Engine/Private/NetConnection.cpp
+++ b/Engine/Source/Runtime/Engine/Private/NetConnection.cpp
@@ -1539,6 +1539,11 @@ void UChildConnection::CleanUp()
 	Driver = NULL;
 }
 
+void UChildConnection::RemoveLocalPlayer(UGameInstance* GameInstance, ULocalPlayer* LocalPlayer)
+{
+	GameInstance->RemoveLocalPlayer(LocalPlayer);
+}
+
 void UChildConnection::CloseAndRemoveChild(UE::Net::ECloseChildFlags Flags)
 {
 	UE_LOG(LogNet, Log, TEXT("[%s:%s] Closing child connection %s."), *GetNameSafe(GetDriver()), *GetNameSafe(Parent), *GetName());
@@ -1576,7 +1581,7 @@ void UChildConnection::CloseAndRemoveChild(UE::Net::ECloseChildFlags Flags)
 					{
 						UE_LOG(LogNet, Log, TEXT("[%s:%s] Remove local player %s."), *GetNameSafe(GetDriver()), *GetNameSafe(Parent), *GetNameSafe(LocalPlayer));
 
-						GameInstance->RemoveLocalPlayer(LocalPlayer);
+						RemoveLocalPlayer(GameInstance, LocalPlayer);
 					}
 				}
 			}
