Refactor UProxyNetDriver::TickFlush() to use AActor::GetNetOwner() instead of AActor::GetOwner().

#Synced-CL 51341769

diff --git a/Engine/Plugins/Runtime/MultiServerReplication/Source/MultiServerReplication/Private/MultiServerProxy.cpp b/Engine/Plugins/Runtime/MultiServerReplication/Source/MultiServerReplication/Private/MultiServerProxy.cpp
--- a/Engine/Plugins/Runtime/MultiServerReplication/Source/MultiServerReplication/Private/MultiServerProxy.cpp
+++ b/Engine/Plugins/Runtime/MultiServerReplication/Source/MultiServerReplication/Private/MultiServerProxy.cpp
@@ -1204,14 +1204,31 @@ void UProxyNetDriver::TickFlush(float DeltaSeconds)
 		for (TActorIterator<AActor> It(World); It; ++It)
 		{
 			const AActor* Actor = *It;
-			const AActor* Owner = Actor->GetOwner();
-			
+
 			if (!Actor->GetIsReplicated() || !(Actor->IsActorBeginningPlay() || Actor->HasActorBegunPlay()))
 			{
 				continue;
 			}
 
-			if (Actor->GetOwner() == nullptr && !Actor->IsA(APlayerController::StaticClass()))
+			// Assumptions:
+			//	* APlayerController::GetNetOwner()		Always returns a pointer to itself (APlayerController)
+			//	* APawn::GetNetOwner()					Always returns a pointer to itself (APawn)
+			//	* AActor::GetNetOwner()					Always return a pointer to either APlayerController or APawn.
+			const AActor* NetOwner = Actor->GetNetOwner();
+
+			ensureMsgf(!Actor->IsA<APlayerController>() || NetOwner == Actor,
+				TEXT("APlayerController::GetNetOwner() should return itself, but %s returned %s"),
+				*Actor->GetName(), *GetNameSafe(NetOwner));
+
+			ensureMsgf(!Actor->IsA<APawn>() || NetOwner == Actor, 
+				TEXT("APawn::GetNetOwner() should return itself, but %s returned %s"),
+				*Actor->GetName(), *GetNameSafe(NetOwner));
+
+			ensureMsgf(NetOwner == nullptr || NetOwner->IsA<APlayerController>() || NetOwner->IsA<APawn>(),
+				TEXT("AActor::GetNetOwner() should return APlayerController or APawn, but %s returned %s (type %s)"),
+				*Actor->GetName(), *NetOwner->GetName(), *NetOwner->GetClass()->GetName());
+
+			if (NetOwner == nullptr)
 			{
 				continue;
 			}
@@ -1219,7 +1236,7 @@ void UProxyNetDriver::TickFlush(float DeltaSeconds)
 			UNetConnection* ProxyConnection = nullptr;
 			for (TPair<TObjectPtr<UNetConnection>, TObjectPtr<const APlayerController>>& Pair : ProxyConnectionToPrimaryGameServerPlayerController)
 			{
-				if (Actor == Pair.Value || Owner == Pair.Value)
+				if (NetOwner == Pair.Value || NetOwner == Pair.Value->GetPawn())
 				{
 					ProxyConnection = Pair.Key;
 					break;
@@ -1694,7 +1711,9 @@ void UProxyNetDriver::SetProxyConnectionAsOwner(UNetConnection* ProxyConnection,
 			{
 				DriverReplicationSystem->SetReplicationEnabledForConnection(ProxyConnection->GetConnectionHandle().GetParentConnectionId(), true);
 				DriverReplicationSystem->SetOwningNetConnection(RefHandle, ConnectionId);
-				DriverReplicationSystem->SetReplicationConditionConnectionFilter(RefHandle, UE::Net::EReplicationCondition::RoleAutonomous, ConnectionId, true);
+				
+				const bool bEnableAutonomousCondition = Actor->GetRemoteRole() == ROLE_AutonomousProxy;
+				DriverReplicationSystem->SetReplicationConditionConnectionFilter(RefHandle, UE::Net::EReplicationCondition::RoleAutonomous, ConnectionId, bEnableAutonomousCondition);
 			}
 		}	
 	}
