MultiServerProxy:  Use GetLocalRole rather than GetRemoteRole as we are executing in the context of a Client (GetLocalRole is what the Server assigned us and GetRemoteRole is always ROLE_Authority).

#Synced-CL 51676086

diff --git a/Engine/Plugins/Runtime/MultiServerReplication/Source/MultiServerReplication/Private/MultiServerProxy.cpp b/Engine/Plugins/Runtime/MultiServerReplication/Source/MultiServerReplication/Private/MultiServerProxy.cpp
--- a/Engine/Plugins/Runtime/MultiServerReplication/Source/MultiServerReplication/Private/MultiServerProxy.cpp
+++ b/Engine/Plugins/Runtime/MultiServerReplication/Source/MultiServerReplication/Private/MultiServerProxy.cpp
@@ -1154,7 +1154,7 @@ bool UProxyNetDriver::CanDowngradeActorRole(UNetConnection* ProxyConnection, AAc
 	{
 		// If this autonomous actor is owned by a player that is bound to the same proxy connection
 		// as the attached proxy player controller, don't downgrade from ROLE_AutonomousProxy to ROLE_SimulatedProxy.
-		if (Actor->GetRemoteRole() == ROLE_AutonomousProxy)
+		if (Actor->GetLocalRole() == ROLE_AutonomousProxy)
 		{
 			if (const UNetConnection* PlayerProxyConnection = const_cast<UProxyNetDriver*>(this)->GetProxyConnectionFromPrimaryPlayer(Player))
 			{
@@ -1711,11 +1711,18 @@ void UProxyNetDriver::SetProxyConnectionAsOwner(UNetConnection* ProxyConnection,
 			{
 				DriverReplicationSystem->SetReplicationEnabledForConnection(ProxyConnection->GetConnectionHandle().GetParentConnectionId(), true);
 				DriverReplicationSystem->SetOwningNetConnection(RefHandle, ConnectionId);
-				
-				const bool bEnableAutonomousCondition = Actor->GetRemoteRole() == ROLE_AutonomousProxy;
+
+				// On the proxy, GetRemoteRole() returns ROLE_Authority (proxy is client of game server). Use GetLocalRole() which mirrors the game server's RemoteRole.
+				const bool bEnableAutonomousCondition = Actor->GetLocalRole() == ROLE_AutonomousProxy;
 				DriverReplicationSystem->SetReplicationConditionConnectionFilter(RefHandle, UE::Net::EReplicationCondition::RoleAutonomous, ConnectionId, bEnableAutonomousCondition);
+
+				// Force net update when ownership is first assigned to ensure Role property re-send
+				if (bEnableAutonomousCondition)
+				{
+					DriverReplicationSystem->ForceNetUpdate(RefHandle);
+				}
 			}
-		}	
+		}
 	}
 }
 
