Use the cvar `net.proxy.Enabled` to configure a game server as a multi-server proxy.

#Synced-CL 52299713

diff --git a/Engine/Plugins/Runtime/MultiServerReplication/Source/MultiServerReplication/Private/MultiServerReplicationModule.cpp b/Engine/Plugins/Runtime/MultiServerReplication/Source/MultiServerReplication/Private/MultiServerReplicationModule.cpp
--- a/Engine/Plugins/Runtime/MultiServerReplication/Source/MultiServerReplication/Private/MultiServerReplicationModule.cpp
+++ b/Engine/Plugins/Runtime/MultiServerReplication/Source/MultiServerReplication/Private/MultiServerReplicationModule.cpp
@@ -2,5 +2,149 @@
 
 #include "MultiServerReplicationModule.h"
 #include "Modules/ModuleManager.h"
+#include "UnrealEngine.h"
+#include "HAL/IConsoleManager.h"
+#include "Misc/CoreDelegates.h"
+#include "Misc/CommandLine.h"
+#include "Misc/Parse.h"
+#include "MultiServerReplicationTypes.h"
+
+namespace UE::MultiServerProxy::Private
+{
+	bool bProxyEnabled = false;
+	static FAutoConsoleVariableRef CVarProxyEnabled(
+		TEXT("net.proxy.Enabled"),
+		bProxyEnabled,
+		TEXT("When true, configure this game server to behave as a multi-server proxy."));
+
+	enum class ENetDriverOverrideState
+	{
+		NoOverride,              // No -NetDriverOverrides specified
+		ProxyDriverSpecified,    // ProxyNetDriver explicitly specified
+		ConflictDetected         // Different driver specified for GameNetDriver
+	};
+
+	// Checks the -NetDriverOverrides command line argument for conflicts with proxy configuration.
+	// Supports the same formats as the engine's NetDriverOverrides parsing:
+	//   -NetDriverOverrides=DriverClassName (targets GameNetDriver by default)
+	//   -NetDriverOverrides="DefName,DriverClassName"
+	//   -NetDriverOverrides="DefName,DriverClassName,FallbackClassName"
+	//   -NetDriverOverrides="DriverClassName;DefName2,DriverClassName2" (multiple entries)
+	//
+	// Returns the state of the -NetDriverOverrides command line argument.
+	static ENetDriverOverrideState CheckForNetDriverOverrideConflict()
+	{
+		FString NetDriverOverridesValue;
+
+		if (!FParse::Value(FCommandLine::Get(), TEXT("-NetDriverOverrides="), NetDriverOverridesValue))
+		{
+			return ENetDriverOverrideState::NoOverride;
+		}
+
+		// Parse entries separated by semicolons (same format as engine code)
+		TArray<FString> OverrideEntries;
+		NetDriverOverridesValue.ParseIntoArray(OverrideEntries, TEXT(";"), false);
+
+		const FString ProxyNetDriverClassName = TEXT("/Script/MultiServerReplication.ProxyNetDriver");
+
+		for (const FString& Entry : OverrideEntries)
+		{
+			TArray<FString> EntryParams;
+			Entry.ParseIntoArray(EntryParams, TEXT(","), false);
+
+			if (EntryParams.Num() == 0)
+			{
+				continue;
+			}
+
+			// Determine the target DefName and DriverClassName based on parameter count
+			// If 1 parameter: targets GameNetDriver by default
+			// If 2+ parameters: first param is DefName, second is DriverClassName
+			FName TargetDefName = (EntryParams.Num() > 1) ? FName(*EntryParams[0]) : NAME_GameNetDriver;
+			FString DriverClassName = (EntryParams.Num() > 1) ? EntryParams[1] : EntryParams[0];
+
+			// Check if this override targets GameNetDriver
+			if (TargetDefName == NAME_GameNetDriver)
+			{
+				if (DriverClassName.Equals(ProxyNetDriverClassName, ESearchCase::IgnoreCase))
+				{
+					return ENetDriverOverrideState::ProxyDriverSpecified;
+				}
+				else
+				{
+					// A different driver is specified - only a conflict if cvar is enabled
+					if (bProxyEnabled)
+					{
+						return ENetDriverOverrideState::ConflictDetected;
+					}
+					else
+					{
+						return ENetDriverOverrideState::NoOverride;
+					}
+				}
+			}
+		}
+
+		return ENetDriverOverrideState::NoOverride;
+	}
+}
+
+void FMultiServerReplicationModule::StartupModule()
+{
+	FCoreDelegates::GetOnPostEngineInit().AddLambda([this]()
+	{
+		SetupProxy();
+	});
+}
+
+void FMultiServerReplicationModule::SetupProxy()
+{
+	using namespace UE::MultiServerProxy::Private;
+
+	// Check if -NetDriverOverrides command line argument conflicts with net.proxy.Enabled cvar
+	ENetDriverOverrideState OverrideState = CheckForNetDriverOverrideConflict();
+
+	switch (OverrideState)
+	{
+	case ENetDriverOverrideState::ConflictDetected:
+		UE_LOG(LogMultiServerReplication, Error,
+			TEXT("Conflict detected: -NetDriverOverrides specifies a different driver for GameNetDriver, but net.proxy.Enabled cvar is set to true. ")
+			TEXT("The command line argument takes precedence. Disabling proxy mode and using the command line override."));
+
+		bProxyEnabled = false;
+		return;
+
+	case ENetDriverOverrideState::ProxyDriverSpecified:
+		bProxyEnabled = true;
+		break;
+	}
+
+	// Apply the proxy override if the cvar is enabled and no conflicting command line override was found
+	if (bProxyEnabled)
+	{
+		if (GEngine)
+		{
+			auto FindGameNetDriverPred = [](const FNetDriverDefinition& Def)
+			{
+				return Def.DefName == NAME_GameNetDriver;
+			};
+
+			FNetDriverDefinition* GameNetDriverDef = GEngine->NetDriverDefinitions.FindByPredicate(FindGameNetDriverPred);
+			if (GameNetDriverDef)
+			{
+				GameNetDriverDef->DriverClassName = FName(TEXT("/Script/MultiServerReplication.ProxyNetDriver"));
+				GameNetDriverDef->DriverClassNameFallback = FName(TEXT("/Script/MultiServerReplication.ProxyNetDriver"));
+			}
+		}
+	}
+}
+
+bool FMultiServerReplicationModule::IsRunningAsProxy()
+{
+	// This function should not be called from the editor or PIE
+	check(UE::GetPlayInEditorID() == INDEX_NONE);
+
+	return UE::MultiServerProxy::Private::bProxyEnabled;
+}
 
 IMPLEMENT_MODULE(FMultiServerReplicationModule, MultiServerReplication);
\ No newline at end of file
diff --git a/Engine/Plugins/Runtime/MultiServerReplication/Source/MultiServerReplication/Public/MultiServerReplicationModule.h b/Engine/Plugins/Runtime/MultiServerReplication/Source/MultiServerReplication/Public/MultiServerReplicationModule.h
--- a/Engine/Plugins/Runtime/MultiServerReplication/Source/MultiServerReplication/Public/MultiServerReplicationModule.h
+++ b/Engine/Plugins/Runtime/MultiServerReplication/Source/MultiServerReplication/Public/MultiServerReplicationModule.h
@@ -5,7 +5,7 @@
 #include "Modules/ModuleInterface.h"
 
 /** See MultiServerNode for more details about multi-server networking as a whole. */
-class FMultiServerReplicationModule : public IModuleInterface
+class MULTISERVERREPLICATION_API FMultiServerReplicationModule : public IModuleInterface
 {
 public:
 
@@ -13,10 +13,34 @@ public:
 	virtual ~FMultiServerReplicationModule() {}
 
 	// IModuleInterface
-	virtual void StartupModule() override {}
+	virtual void StartupModule() override;
 	virtual void ShutdownModule() override {}
 	virtual bool SupportsDynamicReloading() override { return false; }
 	virtual bool SupportsAutomaticShutdown() override { return false; }
+
+	/** 
+	 * Return true if the process is running as a proxy.
+	 *
+	 * This function should *NOT* be used in code that is intended to run from within PIE.
+	 */
+	static bool IsRunningAsProxy();
+
+private:
+
+	/**
+ 	 * Configures a game server to act as a proxy.
+	 *
+	 * A game server is configured as a proxy when the cvar `net.proxy.Enabled` is set to `true` at start up.
+	 *
+	 * These are the methods in which code on a game server can determine if it's running as a proxy:
+	 *		* UReplicationSystem::GetProxyType() or UNetDriver::GetProxyType()
+	 *		* FMultiServerReplicationModule::IsRunningAsProxy()
+	 *		* Read the value of the cvar `net.proxy.Enabled`.
+	 *
+	 * The preferred method is to use `GetProxyType()` which is world-scoped, and only use `IsRunningAsProxy()`
+	 * or reading the cvar directly when `GetProxyType()` is unavailable (e.g. before the driver has been initialized).
+	 */
+	void SetupProxy();
 };
 
 
diff --git a/Engine/Source/Programs/AutomationTool/Scripts/LaunchMultiServer.Automation.cs b/Engine/Source/Programs/AutomationTool/Scripts/LaunchMultiServer.Automation.cs
--- a/Engine/Source/Programs/AutomationTool/Scripts/LaunchMultiServer.Automation.cs
+++ b/Engine/Source/Programs/AutomationTool/Scripts/LaunchMultiServer.Automation.cs
@@ -597,7 +597,7 @@ Example running 4 clients locally connecting to single proxy server w/ a non-loc
 		virtual protected void StartProcessesForProxyServer(string[] LocalProxyServerAddresses, string[] AllGameServerAddresses, string AdditionalProxyArguments)
 		{
 			string CommonProxyServerArgs = AdditionalProxyArguments;
-			CommonProxyServerArgs += " -NetDriverOverrides=/Script/MultiServerReplication.ProxyNetDriver";
+			CommonProxyServerArgs += " -ini:Engine:[ConsoleVariables]:net.proxy.Enabled=1";
 			CommonProxyServerArgs += NoTimeouts ? " -notimeouts" : "";
 			CommonProxyServerArgs += IsAttachingDebugger(ProcessLaunchType.ProxyServer) ? " -WaitForDebuggerNoBreak" : "";
 			CommonProxyServerArgs += " -port={ProxyServerPort}"; // will be Replaced() below.
