Add initial support to the proxy for level streaming and correctly handling level placed actors on DSTM servers.

#Synced-CL 52322957

#Synced-CL 52322957

#Synced-CL 52384223

diff --git a/Engine/Plugins/Runtime/MultiServerReplication/Source/MultiServerReplication/MultiServerReplication.Build.cs b/Engine/Plugins/Runtime/MultiServerReplication/Source/MultiServerReplication/MultiServerReplication.Build.cs
--- a/Engine/Plugins/Runtime/MultiServerReplication/Source/MultiServerReplication/MultiServerReplication.Build.cs
+++ b/Engine/Plugins/Runtime/MultiServerReplication/Source/MultiServerReplication/MultiServerReplication.Build.cs
@@ -23,5 +23,7 @@ public class MultiServerReplication : ModuleRules
 				"OnlineSubsystemUtils"
 			}
 		);
+
+        bAllowUETypesInNamespaces = true;
     }
 }
diff --git a/Engine/Plugins/Runtime/MultiServerReplication/Source/MultiServerReplication/Private/MultiServerProxy.cpp b/Engine/Plugins/Runtime/MultiServerReplication/Source/MultiServerReplication/Private/MultiServerProxy.cpp
--- a/Engine/Plugins/Runtime/MultiServerReplication/Source/MultiServerReplication/Private/MultiServerProxy.cpp
+++ b/Engine/Plugins/Runtime/MultiServerReplication/Source/MultiServerReplication/Private/MultiServerProxy.cpp
@@ -12,6 +12,8 @@
 #include "Net/Core/Misc/NetSubObjectRegistry.h"
 #include "Net/NetworkMetricsDatabase.h"
 #include "Net/NetworkMetricsDefs.h"
+#include "Engine/LevelStreaming.h"
+#include "GameFramework/UpdateLevelVisibilityLevelInfo.h"
 
 #include "Iris/ReplicationSystem/ReplicationSystem.h"
 #include "Net/Iris/ReplicationSystem/EngineReplicationBridge.h"
@@ -21,6 +23,74 @@ DEFINE_LOG_CATEGORY_STATIC(LogNetProxy, Log, All);
 
 namespace UE::MultiServerProxy::Private
 {
+	static void CopyParmsFromFunction(TArray<uint8>& DestParms, const void* SrcData, int32 Size, const UFunction* Func)
+	{
+		DestParms.SetNumZeroed(Size);
+		if (Func)
+		{
+			for (TFieldIterator<FProperty> It(Func); It && It->HasAnyPropertyFlags(CPF_Parm); ++It)
+			{
+				It->CopyCompleteValue_InContainer(DestParms.GetData(), SrcData);
+			}
+		}
+	}
+
+	FPendingBroadcastRPC::FPendingBroadcastRPC(const UFunction* InFunction, const void* InParms, const UObject* InSubObject)
+		: Function(InFunction)
+		, SubObject(InSubObject)
+	{
+		CopyParmsFromFunction(Parms, InParms, InFunction->ParmsSize, InFunction);
+	}
+
+	FPendingBroadcastRPC::FPendingBroadcastRPC(const FPendingBroadcastRPC& Other)
+		: Function(Other.Function)
+		, SubObject(Other.SubObject)
+	{
+		if (Other.Parms.Num() > 0)
+		{
+			CopyParmsFromFunction(Parms, Other.Parms.GetData(), Other.Parms.Num(), Function.ResolveObjectPtr());
+		}
+	}
+
+	FPendingBroadcastRPC& FPendingBroadcastRPC::operator=(const FPendingBroadcastRPC& Other)
+	{
+		if (this != &Other)
+		{
+			DestroyParms();
+			Function = Other.Function;
+			SubObject = Other.SubObject;
+			if (Other.Parms.Num() > 0)
+			{
+				CopyParmsFromFunction(Parms, Other.Parms.GetData(), Other.Parms.Num(), Function.ResolveObjectPtr());
+			}
+		}
+		return *this;
+	}
+
+	FPendingBroadcastRPC::~FPendingBroadcastRPC()
+	{
+		DestroyParms();
+	}
+
+	void FPendingBroadcastRPC::DestroyParms()
+	{
+		if (Parms.Num() == 0)
+		{
+			return;
+		}
+
+		const UFunction* Func = Function.ResolveObjectPtrEvenIfGarbage();
+		if (Func)
+		{
+			for (TFieldIterator<FProperty> It(Func); It && It->HasAnyPropertyFlags(CPF_Parm); ++It)
+			{
+				It->DestroyValue_InContainer(Parms.GetData());
+			}
+		}
+
+		Parms.Reset();
+	}
+
 	FString ConnectionToString(const UNetConnection* Connection)
 	{
 		if (Connection)
@@ -62,6 +132,49 @@ namespace UE::MultiServerProxy::Private
 		TEXT("net.proxy.EnableParentConnectionReplication"),
 		bEnableParentConnectionReplication,
 		TEXT("When true, the proxy will configure the parent connection to the backend game server to replicate state."));
+
+	/**
+	 * The proxy behaves as both a client and a server but there are values in APlayerController that are set based on it being on a
+	 * client or a server.
+	 *
+	 * By default, these values in the player controller are set as if they were on a client because they were replicated to the client
+	 * from the game server, but this class can be used to temporarily configure these values as if they were on a server.
+	 */
+	struct FScopedPlayerControllerPlayerSwap
+	{
+		FScopedPlayerControllerPlayerSwap(APlayerController* InPlayerController, UProxyNetDriver* InProxyNetDriver)
+			: PlayerController(InPlayerController)
+		{
+			if (PlayerController && InProxyNetDriver)
+			{
+				if (ULocalPlayer* OwningPlayer = Cast<ULocalPlayer>(PlayerController->Player))
+				{
+					if (UNetConnection* ProxyConnection = InProxyNetDriver->GetProxyConnectionForPlayer(OwningPlayer))
+					{
+						OriginalPlayer = PlayerController->Player;
+						PlayerController->Player = ProxyConnection;
+					}
+				}
+			}
+		}
+
+		~FScopedPlayerControllerPlayerSwap()
+		{
+			if (OriginalPlayer)
+			{
+				PlayerController->Player = OriginalPlayer;
+			}
+		}
+
+		FScopedPlayerControllerPlayerSwap(const FScopedPlayerControllerPlayerSwap&) = delete;
+		FScopedPlayerControllerPlayerSwap& operator=(const FScopedPlayerControllerPlayerSwap&) = delete;
+		FScopedPlayerControllerPlayerSwap(FScopedPlayerControllerPlayerSwap&&) = delete;
+		FScopedPlayerControllerPlayerSwap& operator=(FScopedPlayerControllerPlayerSwap&&) = delete;
+
+	private:
+		APlayerController* PlayerController = nullptr;
+		TObjectPtr<UPlayer> OriginalPlayer;
+	};
 }
 
 namespace UE::MultiServerProxy::Metrics
@@ -620,6 +733,14 @@ void UProxyBackendNetDriver::ForwardRemoteFunction(UObject* RootObject, UObject*
 			ProxyNetDriver->InternalProcessRemoteFunction(OwningActor, SubObject, ProxyConnection, Function, Parms, nullptr, nullptr, true);
 		});
 	}
+
+	const FSoftObjectPath FunctionPath(Function->GetOwnerClass()->GetClassPathName(), Function->GetName());
+	if (ProxyNetDriver->AllowListedProxyFunctions.Contains(FunctionPath))
+	{
+		UE_LOGF(LogNetProxy, Verbose, "RPC: Executing allow-listed %ls on %ls on proxy (game server to client) from %ls", *Function->GetName(), *GetNameSafe(RootObject), *GetName());
+		UE::MultiServerProxy::Private::FScopedPlayerControllerPlayerSwap PlayerSwap(Cast<APlayerController>(RootObject), ProxyNetDriver);
+		RootObject->ProcessEvent(Function, Parms);
+	}
 }
 
 bool UProxyBackendNetDriver::ShouldSkipRepNotifies() const
@@ -675,6 +796,14 @@ void UProxyBackendNetDriver::OnIrisRootObjectPostInit(UE::Net::FNetRefHandle Roo
 
 	UE_LOG(LogNetProxy, Verbose, TEXT("[%s] Attaching root object %s (%s) to %s"), *GetNameSafe(this), *GetNameSafe(RootObject), *RootHandle.ToString(), *GetNameSafe(ProxyNetDriver));
 	FrontendBridge->StartReplicatingRootObject(RootObject, UObjectReplicationBridge::FRootObjectReplicationParams(), FactoryId);
+
+	if (AActor* Actor = Cast<AActor>(RootObject))
+	{
+		if (UEngineReplicationBridge* EngineFrontendBridge = Cast<UEngineReplicationBridge>(FrontendBridge))
+		{
+			EngineFrontendBridge->AddRootObjectToContainerGroup(Actor, Actor->GetLevel());
+		}
+	}
 }
 
 void UProxyBackendNetDriver::OnIrisRootObjectDetached(UE::Net::FNetRefHandle RootHandle, UObject* RootObject, UE::Net::EDetachReason Reason)
@@ -906,6 +1035,8 @@ bool UProxyNetDriver::InitBase(bool bInitAsClient, FNetworkNotify* InNotify, con
 	if (bSuccess)
 	{
 		NetConnectionClass = UProxyNetConnection::StaticClass();
+		ProxyPreGarbageCollectHandle = FCoreUObjectDelegates::GetPreGarbageCollectDelegate().AddUObject(this, &UProxyNetDriver::PreGarbageCollect);
+		ProxyPostGarbageCollectHandle = FCoreUObjectDelegates::GetPostGarbageCollect().AddUObject(this, &UProxyNetDriver::PostGarbageCollect);
 	}
 
 	FString GameServerAddresses;
@@ -943,6 +1074,16 @@ bool UProxyNetDriver::InitBase(bool bInitAsClient, FNetworkNotify* InNotify, con
 
 	SetReplicateTransactionally(false);
 
+	// These functions need to be executed on the proxy to update Iris filter groups used for streaming on the frontend replication system.
+	AllowListedProxyFunctions.Add(FSoftObjectPath(APlayerController::StaticClass()->GetClassPathName(), GET_FUNCTION_NAME_CHECKED(APlayerController, ServerUpdateLevelVisibility).ToString()));
+	AllowListedProxyFunctions.Add(FSoftObjectPath(APlayerController::StaticClass()->GetClassPathName(), GET_FUNCTION_NAME_CHECKED(APlayerController, ServerUpdateMultipleLevelsVisibility).ToString()));
+
+	// RPCs from the client that must be forwarded to all game servers (not just the primary).
+
+	// These functions need to be forwarded to game servers to know that the player can see a particular level.
+	BroadcastToAllGameServersFunctions.Add(FSoftObjectPath(APlayerController::StaticClass()->GetClassPathName(), GET_FUNCTION_NAME_CHECKED(APlayerController, ServerUpdateLevelVisibility).ToString()));
+	BroadcastToAllGameServersFunctions.Add(FSoftObjectPath(APlayerController::StaticClass()->GetClassPathName(), GET_FUNCTION_NAME_CHECKED(APlayerController, ServerUpdateMultipleLevelsVisibility).ToString()));
+
 	GuidCache = TSharedPtr<FNetGUIDCache>(new FProxyNetGUIDCache(this));
 
 	// All of the instances of UProxyBackendNetDriver will use this shared GUID cache to ensure that objects
@@ -1046,8 +1187,24 @@ void UProxyNetDriver::ForwardRemoteFunction(UObject* RootObject, UObject* SubObj
 		}
 	}
 
-	if (TObjectPtr<const APlayerController>* GamePlayerControllerPtr = PlayerToPrimaryGameServerPlayerController.Find(OwningPlayer))
+	const FSoftObjectPath FunctionPath(Function->GetOwnerClass()->GetClassPathName(), Function->GetName());
+
+	if (bShouldForwardRPC && AllowListedProxyFunctions.Contains(FunctionPath))
+	{
+		UE_LOGF(LogNetProxy, Verbose, "RPC: Executing allow-listed %ls on %ls on proxy (client to game server) for player %ls", *Function->GetName(), *GetNameSafe(OwningActor), *GetNameSafe(OwningPlayer));
+
+		UE::MultiServerProxy::Private::FScopedPlayerControllerPlayerSwap PlayerSwap(Cast<APlayerController>(OwningActor), this);
+		OwningActor->ProcessEvent(Function, Parms);
+	}
+
+	const bool bShouldBroadcastToAllGameServers = bShouldForwardRPC && BroadcastToAllGameServersFunctions.Contains(FunctionPath);
+	if (bShouldBroadcastToAllGameServers)
 	{
+		BroadcastToAllClientConnectedGameServers(OwningPlayer, Function, Parms, SubObject);
+	}
+	else if (TObjectPtr<const APlayerController>* GamePlayerControllerPtr = PlayerToPrimaryGameServerPlayerController.Find(OwningPlayer))
+	{
+		// Forward all other RPCs to the primary game server only.
 		UNetConnection* GameConnection = GamePlayerControllerPtr->Get()->NetConnection;
 		if (GameConnection)
 		{
@@ -1080,6 +1237,53 @@ void UProxyNetDriver::ForwardRemoteFunction(UObject* RootObject, UObject* SubObj
 	}
 }
 
+UNetConnection* UProxyNetDriver::GetProxyConnectionForPlayer(const ULocalPlayer* Player) const
+{
+	if (const TObjectPtr<UNetConnection>* ProxyConnectionPtr = PlayerToProxyConnection.Find(Player))
+	{
+		return ProxyConnectionPtr->Get();
+	}
+	return nullptr;
+}
+
+void UProxyNetDriver::BroadcastToAllClientConnectedGameServers(const ULocalPlayer* Player, const UFunction* Function, const void* Parms, const UObject* SubObject)
+{
+	TObjectPtr<UNetConnection>* ProxyConnectionPtr = PlayerToProxyConnection.Find(Player);
+	if (!ProxyConnectionPtr || !*ProxyConnectionPtr)
+	{
+		UE_LOGF(LogNetProxy, Warning, "BroadcastToAllClientConnectedGameServers: No proxy connection found for player %ls", *GetNameSafe(Player));
+		return;
+	}
+
+	UNetConnection* ProxyConnection = *ProxyConnectionPtr;
+	for (TPair<uint32, FMultiServerProxyInternalConnectionRoute>& Pair : InternalRoutes)
+	{
+		if (Pair.Value.ProxyConnection != ProxyConnection)
+		{
+			continue;
+		}
+
+		const bool bIsConnected = (Pair.Value.State == EProxyConnectionState::Connected || Pair.Value.State == EProxyConnectionState::ConnectedPrimary);
+		if (bIsConnected)
+		{
+			UNetConnection* GameServerConn = Pair.Value.GameServerConnection;
+			APlayerController* PC = const_cast<APlayerController*>(Pair.Value.PlayerController.Get());
+
+			UE_LOGF(LogNetProxy, Verbose, "RPC: Broadcasting %ls on %ls to game server connection %ls:%ls",
+				*GetNameSafe(Function), *GetNameSafe(PC), *GameServerConn->Driver->GetName(), *GetNameSafe(GameServerConn));
+
+			GameServerConn->Driver->ProcessRemoteFunction(PC, const_cast<UFunction*>(Function), const_cast<void*>(Parms), static_cast<FOutParmRec*>(nullptr), static_cast<FFrame*>(nullptr), const_cast<UObject*>(SubObject));
+		}
+		else
+		{
+			Pair.Value.PendingBroadcastRPCs.Emplace(Function, Parms, SubObject);
+
+			UE_LOGF(LogNetProxy, Verbose, "RPC: Queued broadcast %ls for route [%u] (not yet connected, state=%d, pending=%d)",
+				   *GetNameSafe(Function), Pair.Value.ClientHandshakeId, static_cast<uint8>(Pair.Value.State), Pair.Value.PendingBroadcastRPCs.Num());
+		}
+	}
+}
+
 bool UProxyNetDriver::ShouldReplicateFunction(AActor* Actor, UFunction* Function) const
 {
 	// If any game code in the proxy world attempts to send an RPC it should not be called.
@@ -1187,6 +1391,44 @@ void UProxyNetDriver::Shutdown()
 	GameServerConnections.Reset();
 }
 
+void UProxyNetDriver::FinishDestroy()
+{
+	FCoreUObjectDelegates::GetPostGarbageCollect().Remove(ProxyPostGarbageCollectHandle);
+	Super::FinishDestroy();
+}
+
+void UProxyNetDriver::PreGarbageCollect()
+{
+	for (TPair<uint32, FMultiServerProxyInternalConnectionRoute>& ClientHandshakeIdToRoute : InternalRoutes)
+	{
+		FMultiServerProxyInternalConnectionRoute* Route = &ClientHandshakeIdToRoute.Value;
+
+		for (auto It = Route->PendingBroadcastRPCs.CreateIterator(); It; ++It)
+		{
+			if (It->Function.ResolveObjectPtr() == nullptr)
+			{
+				UNetConnection* GameServerConn = Route->GameServerConnection;
+
+				UE_LOGF(LogNetProxy, Warning, "Dropping garbage collected function on %ls to game server connection %ls:%ls",
+					   *GetNameSafe(Route->PlayerController), *GameServerConn->Driver->GetName(), *GetNameSafe(GameServerConn));
+
+				It.RemoveCurrent();
+			}
+		}
+	}
+}
+
+void UProxyNetDriver::PostGarbageCollect()
+{
+	for (auto It = PlayerToProxyConnection.CreateIterator(); It; ++It)
+	{
+		if (!It.Key().ResolveObjectPtr())
+		{
+			It.RemoveCurrent();
+		}
+	}
+}
+
 void UProxyNetDriver::TickFlush(float DeltaSeconds)
 {
 	TArray<TWeakObjectPtr<APlayerController>> PlayerControllerNetConnectionsToRestore;
@@ -1399,6 +1641,7 @@ void UProxyNetDriver::HandleClosedProxyConnection(UNetConnection* ProxyConnectio
 		const bool bStillConnecting  = (Route->State == EProxyConnectionState::Connecting || Route->State == EProxyConnectionState::ConnectingPrimary);
 		if (!bStillConnecting)
 		{
+			PlayerToProxyConnection.Remove(Route->Player);
 			FMultiServerProxyInternalConnectionRoute RemovedRoute;
 			if (InternalRoutes.RemoveAndCopyValue(ClientHandshakeIdToRemove, RemovedRoute))
 			{
@@ -1512,8 +1755,11 @@ void UProxyNetDriver::ProcessConnectionState_Connecting(FMultiServerProxyInterna
 
 	FinalizeGameServerConnection(Route, GameServerConnection, PlayerController);
 
+	PlayerToProxyConnection.Add(Route->Player, Route->ProxyConnection);
 	Route->State = EProxyConnectionState::Connected;
 
+	FlushPendingBroadcastRPCs(Route);
+
 	UE_LOG(LogNetProxy, Log, TEXT("[%d] Successfully connected proxy connection %s to non-primary game server connection %s using player controller %s and player %s."),
 		   Route->ClientHandshakeId,
 		   *UE::MultiServerProxy::Private::ConnectionToString(Route->ProxyConnection),
@@ -1534,10 +1780,13 @@ void UProxyNetDriver::ProcessConnectionState_ConnectingPrimary(FMultiServerProxy
 	Route->ProxyConnection->LastReceiveTime = GetElapsedTime();
 
 	PlayerToPrimaryGameServerPlayerController.Add(Route->Player, Route->PlayerController);
+	PlayerToProxyConnection.Add(Route->Player, Route->ProxyConnection);
 	ProxyConnectionToPrimaryGameServerPlayerController.Add(Route->ProxyConnection, Route->PlayerController);
 
 	Route->State = EProxyConnectionState::ConnectedPrimary;
 
+	FlushPendingBroadcastRPCs(Route);
+
 	UE_LOG(LogNetProxy, Log, TEXT("[%d] Successfully connected proxy connection %s to primary game server connection %s using player controller %s and player %s."),
 		   Route->ClientHandshakeId,
 		   *UE::MultiServerProxy::Private::ConnectionToString(Route->ProxyConnection),
@@ -1576,6 +1825,7 @@ void UProxyNetDriver::ProcessConnectionState_PendingClose(FMultiServerProxyInter
 	Route->GameServerConnection = GameServerConnection;
 	GameServerConnection->PlayerController = PlayerController;
 
+	PlayerToProxyConnection.Remove(Route->Player);
 	FMultiServerProxyInternalConnectionRoute RemovedRoute;
 	if (InternalRoutes.RemoveAndCopyValue(Route->ClientHandshakeId, RemovedRoute))
 	{
@@ -1819,6 +2069,46 @@ void UProxyNetDriver::FlushSplitJoinRequests(FGameServerConnectionState* GameSer
 	GameServerConnectionState->PendingSplitJoinRequests.Reset();
 }
 
+void UProxyNetDriver::FlushPendingBroadcastRPCs(FMultiServerProxyInternalConnectionRoute* Route)
+{
+	if (Route->PendingBroadcastRPCs.Num() == 0)
+	{
+		return;
+	}
+
+	if (!Route->GameServerConnection || !Route->PlayerController)
+	{
+		return;
+	}
+
+	UNetConnection* GameServerConn = Route->GameServerConnection;
+	APlayerController* PC = const_cast<APlayerController*>(Route->PlayerController.Get());
+
+	UE_LOGF(LogNetProxy, Verbose, "Flushing %d pending broadcast RPCs for route [%u] on connection %ls:%ls.",
+		   Route->PendingBroadcastRPCs.Num(),
+		   Route->ClientHandshakeId,
+		   *GameServerConn->Driver->GetName(),
+		   *GetNameSafe(GameServerConn));
+
+	for (UE::MultiServerProxy::Private::FPendingBroadcastRPC& PendingRPC : Route->PendingBroadcastRPCs)
+	{
+		if (UFunction* Function = const_cast<UFunction*>(PendingRPC.Function.ResolveObjectPtr()))
+		{
+			UE_LOGF(LogNetProxy, Verbose, "RPC: Broadcasting queued %ls on %ls to game server connection %ls:%ls",
+				   *GetNameSafe(Function), *GetNameSafe(PC), *GameServerConn->Driver->GetName(), *GetNameSafe(GameServerConn));
+
+			GameServerConn->Driver->ProcessRemoteFunction(PC, const_cast<UFunction*>(Function), PendingRPC.Parms.GetData(), static_cast<FOutParmRec*>(nullptr), static_cast<FFrame*>(nullptr), const_cast<UObject*>(PendingRPC.SubObject.ResolveObjectPtr()));
+		}
+		else
+		{
+			UE_LOGF(LogNetProxy, Verbose, "RPC: Dropping queued invalid function on %ls to game server connection %ls:%ls",
+				   *GetNameSafe(PC), *GameServerConn->Driver->GetName(), *GetNameSafe(GameServerConn));
+		}
+	}
+
+	Route->PendingBroadcastRPCs.Reset();
+}
+
 UNetConnection* UProxyNetDriver::GetProxyConnectionFromPrimaryPlayer(ULocalPlayer* Player)
 {
 	TObjectPtr<const APlayerController>* GameServerPlayerControllerPtr = PlayerToPrimaryGameServerPlayerController.Find(Player);
@@ -1961,4 +2251,4 @@ void UProxyNetConnection::CleanUp()
 	}
 
 	Super::CleanUp();
-}
\ No newline at end of file
+}
diff --git a/Engine/Plugins/Runtime/MultiServerReplication/Source/MultiServerReplication/Public/MultiServerProxy.h b/Engine/Plugins/Runtime/MultiServerReplication/Source/MultiServerReplication/Public/MultiServerProxy.h
--- a/Engine/Plugins/Runtime/MultiServerReplication/Source/MultiServerReplication/Public/MultiServerProxy.h
+++ b/Engine/Plugins/Runtime/MultiServerReplication/Source/MultiServerReplication/Public/MultiServerProxy.h
@@ -10,8 +10,34 @@
 #include "Iris/ReplicationSystem/NetObjectFactory.h"
 #include "IpNetDriver.h"
 #include "IpConnection.h"
+#include "UObject/SoftObjectPath.h"
 #include "MultiServerProxy.generated.h"
 
+namespace UE::MultiServerProxy::Private
+{
+
+	struct MULTISERVERREPLICATION_API FPendingBroadcastRPC
+	{
+		TObjectKey<const UFunction> Function;
+
+		TArray<uint8> Parms;
+
+		TObjectKey<const UObject> SubObject;
+
+		FPendingBroadcastRPC() = default;
+		FPendingBroadcastRPC(FPendingBroadcastRPC&& Other) = default;
+		FPendingBroadcastRPC& operator=(FPendingBroadcastRPC&& Other) = default;
+
+		FPendingBroadcastRPC(const UFunction* InFunction, const void* InParms, const UObject* InSubObject);
+		FPendingBroadcastRPC(const FPendingBroadcastRPC& Other);
+		FPendingBroadcastRPC& operator=(const FPendingBroadcastRPC& Other);
+		~FPendingBroadcastRPC();
+
+		void DestroyParms();
+	};
+
+} // namespace UE::MultiServerProxy::Private
+
 enum class EProxyConnectionState : uint8
 {
 	Disconnected,
@@ -125,6 +151,9 @@ struct FMultiServerProxyInternalConnectionRoute
 	/** The player used to represent this route in the game systems. */
 	UPROPERTY()
 	TObjectPtr<class ULocalPlayer> Player;
+	
+	/** Queue to hold RPCs that will be broadcast to all game servers using BroadcastToAllClientConnectedGameServers(). */
+	TArray<UE::MultiServerProxy::Private::FPendingBroadcastRPC> PendingBroadcastRPCs;
 
 	/**
 	 * The player controller.
@@ -371,6 +400,7 @@ public:
 	virtual int32 ServerReplicateActors(float DeltaSeconds) override;
 	virtual bool CanDowngradeActorRole(UNetConnection* ProxyConnection, AActor* Actor) const override;
 	virtual void Shutdown() override;
+	virtual void FinishDestroy() override;
 	virtual void TickFlush(float DeltaSeconds) override;
 	
 	/** Get the next identifier for outgoing connections to game servers. */
@@ -384,13 +414,22 @@ public:
 
 	/** Return a pointer to the shared NetGUID cache used by backend network connections. */
 	TSharedPtr<FNetGUIDCache> GetSharedBackendNetGuidCache();
-	
+
+	/** Forward an RPC to all primary and non-primary game servers for a proxy client connection using PlayerToProxyConnection. */
+	void BroadcastToAllClientConnectedGameServers(const ULocalPlayer* Player, const UFunction* Function, const void* Parms, const UObject* SubObject);
+
+	/** Return the proxy connection for a local player, or nullptr if not found. */
+	UNetConnection* GetProxyConnectionForPlayer(const ULocalPlayer* Player) const;
+
 	/** Called when an incoming proxy connection has been closed and closes connections to backend servers. */
 	void HandleClosedProxyConnection(UNetConnection* ProxyConnection);
 	
 	/** Send any split join requests that have been queued up to the game server. */
 	void FlushSplitJoinRequests(FGameServerConnectionState* GameServerConnectionState);
 
+	/** Flush any broadcast RPCs that were queued while a route was not yet connected. */
+	void FlushPendingBroadcastRPCs(FMultiServerProxyInternalConnectionRoute* Route);
+
 private:
 
 	friend class UProxyListenerNotify;
@@ -399,7 +438,12 @@ private:
 	friend class FNetTestProxyConnectionRouting;
 	friend class FNetTestProxyCloseSecondClientBeforeConnectedToBackend;
 	friend class FNetTestProxyCloseFirstClientBeforeBackendParentConnected;
-	
+	friend class FNetTestProxyAllowListExecOnProxy;
+	friend class FNetTestProxyNonAllowListedSkipsProxy;
+	friend class FNetTestProxyNonBroadcastPrimaryOnly;
+	friend class FNetTestProxyBroadcastReachesAllServers;
+	friend class FNetTestProxyGCCleansPendingBroadcastRPCs;
+
 	friend class UProxyBackendNetDriver;
 
 	/** Close a route's connection and perform any clean up. */
@@ -422,6 +466,10 @@ private:
 	void AggregateBackendNetDriverMetrics();
 	void ResetProxyNetworkMetrics();
 
+	/** Clean up any internal data structures that hold weak object references that are now invalid. */
+	void PreGarbageCollect();
+	void PostGarbageCollect();
+
 	/** 
 	 * Send an RPC correct proxy connection based on the actor's owner and local player.
 	 *
@@ -482,6 +530,12 @@ private:
 	UPROPERTY()
 	TMap<TObjectPtr<class ULocalPlayer>, TObjectPtr<const APlayerController>> PlayerToPrimaryGameServerPlayerController;
 
+	/** Map a local player to it's proxy net connection. Multiple players will point to the same proxy net connection. */
+	TMap<TObjectKey<class ULocalPlayer>, TObjectPtr<UNetConnection>> PlayerToProxyConnection;
+
+	FDelegateHandle ProxyPostGarbageCollectHandle;
+	FDelegateHandle ProxyPreGarbageCollectHandle;
+
 	UPROPERTY()
 	TMap<TObjectPtr<UNetConnection>, FPlayerControllerReassignment> DeferredReassignment;
 
@@ -512,6 +566,12 @@ private:
 
 	// As a client connects to the proxy, randomize PrimaryGameServerForNextClient to one of the known game servers.
 	bool bRandomizePrimaryGameServerForNextClient = false;
+
+	/** RPCs from the client that must be forwarded to all game servers (not just the primary). */
+	TSet<FSoftObjectPath> BroadcastToAllGameServersFunctions;
+
+	/** RPCs that are allow-listed to be executed on the proxy when forwarding (both client->server and server->client). */
+	TSet<FSoftObjectPath> AllowListedProxyFunctions;
 };
 
 /** 
diff --git a/Engine/Source/Runtime/Engine/Private/World.cpp b/Engine/Source/Runtime/Engine/Private/World.cpp
--- a/Engine/Source/Runtime/Engine/Private/World.cpp
+++ b/Engine/Source/Runtime/Engine/Private/World.cpp
@@ -4919,6 +4919,7 @@ void UWorld::SpawnAndSetNoPawnPlayerController(UNetConnection* Connection, uint8
 	Connection->PlayerController->NetConnection = Connection;
 	Connection->PlayerController->FinishSpawning(FTransform());
 
+	Connection->PlayerController->SetPlayer(Connection);
 	Connection->PlayerController->SetRole(ROLE_Authority);
 	Connection->PlayerController->SetReplicates(true);
 	Connection->PlayerController->SetAutonomousProxy(true);
